- Split OnlineServicesEOS into OnlineServicesEOS (EAS+EOS) and OnlineServicesEOSGS (GS == Game Services, i.e. only EOS)

diff --git a/Engine/Plugins/Online/EOSShared/Source/EOSShared/Private/EOSShared.cpp b/Engine/Plugins/Online/EOSShared/Source/EOSShared/Private/EOSShared.cpp
--- a/Engine/Plugins/Online/EOSShared/Source/EOSShared/Private/EOSShared.cpp
+++ b/Engine/Plugins/Online/EOSShared/Source/EOSShared/Private/EOSShared.cpp
@@ -3,10 +3,16 @@
 #include "EOSShared.h"
 #include "EOSSharedTypes.h"
 
+#include "eos_auth_types.h"
 #include "eos_friends_types.h"
 
 DEFINE_LOG_CATEGORY(LogEOSSDK);
 
+FString LexToString(const EOS_EResult EosResult)
+{
+	return UTF8_TO_TCHAR(EOS_EResult_ToString(EosResult));
+}
+
 FString LexToString(const EOS_ProductUserId UserId)
 {
 	FString Result;
@@ -49,4 +55,113 @@ const TCHAR* LexToString(const EOS_EFriendsStatus FriendStatus)
 		case EOS_EFriendsStatus::EOS_FS_InviteReceived: return TEXT("InviteReceived");
 		case EOS_EFriendsStatus::EOS_FS_Friends:		return TEXT("Friends");
 	}
-} 
\ No newline at end of file
+}
+
+bool LexFromString(EOS_EExternalCredentialType& OutEnum, const TCHAR* InString)
+{
+	if (FCString::Stricmp(InString, TEXT("Steam")) == 0)
+	{
+		OutEnum = EOS_EExternalCredentialType::EOS_ECT_STEAM_APP_TICKET;
+	}
+	else if (FCString::Stricmp(InString, TEXT("PSN")) == 0)
+	{
+		OutEnum = EOS_EExternalCredentialType::EOS_ECT_PSN_ID_TOKEN;
+	}
+	else if (FCString::Stricmp(InString, TEXT("Xbox")) == 0)
+	{
+		OutEnum = EOS_EExternalCredentialType::EOS_ECT_XBL_XSTS_TOKEN;
+	}
+	else if (FCString::Stricmp(InString, TEXT("Nintendo")) == 0)
+	{
+		OutEnum = EOS_EExternalCredentialType::EOS_ECT_NINTENDO_ID_TOKEN;
+	}
+	else if (FCString::Stricmp(InString, TEXT("NSA")) == 0)
+	{
+		OutEnum = EOS_EExternalCredentialType::EOS_ECT_NINTENDO_NSA_ID_TOKEN;
+	}
+	else if (FCString::Stricmp(InString, TEXT("Apple")) == 0)
+	{
+		OutEnum = EOS_EExternalCredentialType::EOS_ECT_APPLE_ID_TOKEN;
+	}
+	else if (FCString::Stricmp(InString, TEXT("Google")) == 0)
+	{
+		OutEnum = EOS_EExternalCredentialType::EOS_ECT_GOOGLE_ID_TOKEN;
+	}
+	else
+	{
+		// Unknown means OpenID
+		OutEnum = EOS_EExternalCredentialType::EOS_ECT_OPENID_ACCESS_TOKEN;
+	}
+
+	return true;
+}
+
+bool LexFromString(EOS_EAuthScopeFlags& OutEnum, const TCHAR* InString)
+{
+	if (FCString::Stricmp(InString, TEXT("BasicProfile")) == 0)
+	{
+		OutEnum = EOS_EAuthScopeFlags::EOS_AS_BasicProfile;
+	}
+	else if (FCString::Stricmp(InString, TEXT("FriendsList")) == 0)
+	{
+		OutEnum = EOS_EAuthScopeFlags::EOS_AS_FriendsList;
+	}
+	else if (FCString::Stricmp(InString, TEXT("Presence")) == 0)
+	{
+		OutEnum = EOS_EAuthScopeFlags::EOS_AS_Presence;
+	}
+	else if (FCString::Stricmp(InString, TEXT("FriendsManagement")) == 0)
+	{
+		OutEnum = EOS_EAuthScopeFlags::EOS_AS_FriendsManagement;
+	}
+	else if (FCString::Stricmp(InString, TEXT("Email")) == 0)
+	{
+		OutEnum = EOS_EAuthScopeFlags::EOS_AS_Email;
+	}
+	else if (FCString::Stricmp(InString, TEXT("NoFlags")) == 0 || FCString::Stricmp(InString, TEXT("None")) == 0)
+	{
+		OutEnum = EOS_EAuthScopeFlags::EOS_AS_NoFlags;
+	}
+	else
+	{
+		return false;
+	}
+	return true;
+}
+
+bool LexFromString(EOS_ELoginCredentialType& OutEnum, const TCHAR* InString)
+{
+	if (FCString::Stricmp(InString, TEXT("ExchangeCode")) == 0)
+	{
+		OutEnum = EOS_ELoginCredentialType::EOS_LCT_ExchangeCode;
+	}
+	else if (FCString::Stricmp(InString, TEXT("PersistentAuth")) == 0)
+	{
+		OutEnum = EOS_ELoginCredentialType::EOS_LCT_PersistentAuth;
+	}
+	else if (FCString::Stricmp(InString, TEXT("Password")) == 0)
+	{
+		OutEnum = EOS_ELoginCredentialType::EOS_LCT_Password;
+	}
+	else if (FCString::Stricmp(InString, TEXT("Developer")) == 0)
+	{
+		OutEnum = EOS_ELoginCredentialType::EOS_LCT_Developer;
+	}
+	else if (FCString::Stricmp(InString, TEXT("RefreshToken")) == 0)
+	{
+		OutEnum = EOS_ELoginCredentialType::EOS_LCT_RefreshToken;
+	}
+	else if (FCString::Stricmp(InString, TEXT("AccountPortal")) == 0)
+	{
+		OutEnum = EOS_ELoginCredentialType::EOS_LCT_AccountPortal;
+	}
+	else if (FCString::Stricmp(InString, TEXT("ExternalAuth")) == 0)
+	{
+		OutEnum = EOS_ELoginCredentialType::EOS_LCT_ExternalAuth;
+	}
+	else
+	{
+		return false;
+	}
+	return true;
+}
diff --git a/Engine/Plugins/Online/EOSShared/Source/EOSShared/Public/EOSShared.h b/Engine/Plugins/Online/EOSShared/Source/EOSShared/Public/EOSShared.h
--- a/Engine/Plugins/Online/EOSShared/Source/EOSShared/Public/EOSShared.h
+++ b/Engine/Plugins/Online/EOSShared/Source/EOSShared/Public/EOSShared.h
@@ -21,18 +21,19 @@
 #endif
 
 #define EOS_ENUM_FORWARD_DECL(name) enum class name : int32_t;
+EOS_ENUM_FORWARD_DECL(EOS_EAuthScopeFlags);
 EOS_ENUM_FORWARD_DECL(EOS_EFriendsStatus);
+EOS_ENUM_FORWARD_DECL(EOS_ELoginCredentialType);
 #undef EOS_ENUM_FORWARD_DECL
 
 DECLARE_LOG_CATEGORY_EXTERN(LogEOSSDK, Log, All);
 
-inline FString LexToString(EOS_EResult EosResult)
-{
-	return ANSI_TO_TCHAR(EOS_EResult_ToString(EosResult));
-}
-
+EOSSHARED_API FString LexToString(const EOS_EResult EosResult);
 EOSSHARED_API FString LexToString(const EOS_ProductUserId UserId);
 EOSSHARED_API FString LexToString(const EOS_EpicAccountId AccountId);
 EOSSHARED_API const TCHAR* LexToString(const EOS_EFriendsStatus FriendStatus);
+EOSSHARED_API bool LexFromString(EOS_EAuthScopeFlags& OutEnum, const TCHAR* InString);
+EOSSHARED_API bool LexFromString(EOS_EExternalCredentialType& OutEnum, const TCHAR* InString);
+EOSSHARED_API bool LexFromString(EOS_ELoginCredentialType& OutEnum, const TCHAR* InString);
 
 #endif // WITH_EOS_SDK
\ No newline at end of file
diff --git a/Engine/Plugins/Online/Experimental/OnlineServices/Source/OnlineServicesCommon/Public/Online/OnlineIdCommon.h b/Engine/Plugins/Online/Experimental/OnlineServices/Source/OnlineServicesCommon/Public/Online/OnlineIdCommon.h
--- a/Engine/Plugins/Online/Experimental/OnlineServices/Source/OnlineServicesCommon/Public/Online/OnlineIdCommon.h
+++ b/Engine/Plugins/Online/Experimental/OnlineServices/Source/OnlineServicesCommon/Public/Online/OnlineIdCommon.h
@@ -9,73 +9,76 @@ namespace UE::Online {
 /**
  * A net id registry suitable for use with trivial immutable keys
  */
-template<typename IdType, typename IdValueType, EOnlineServices SubsystemType>
+template<typename IdType, typename IdValueType, EOnlineServices OnlineServicesType>
 class TOnlineBasicIdRegistry
 {
 public:
+	virtual ~TOnlineBasicIdRegistry() = default;
+
 	using HandleType = TOnlineIdHandle<IdType>;
 
-	static TOnlineBasicIdRegistry& Get();
+	const HandleType InvalidHandle = HandleType(OnlineServicesType, 0);
 
-	HandleType FindOrAddHandle(const IdValueType& IdValue)
+	HandleType FindHandle(const IdValueType& IdValue) const
 	{
-		HandleType Handle;
-		// Take a read lock and check if we already have a handle
+		FReadScopeLock ReadLock(Lock);
+		if (const HandleType* FoundHandle = IdValueToHandleMap.Find(IdValue))
 		{
-			FReadScopeLock ReadLock(Lock);
-			if (const uint32* FoundHandle = IdValueToHandleMap.Find(IdValue))
-			{
-				Handle.Handle = *FoundHandle;
-			}
+			return *FoundHandle;
 		}
+		return InvalidHandle;
+	}
 
+	HandleType FindOrAddHandle(const IdValueType& IdValue)
+	{
+		HandleType Handle = FindHandle(IdValue);
 		if (!Handle.IsValid())
 		{
 			// Take a write lock, check again if we already have a handle, or insert a new element.
 			FWriteScopeLock WriteLock(Lock);
 			if (const HandleType* FoundHandle = IdValueToHandleMap.Find(IdValue))
 			{
-				Handle.Handle = *FoundHandle;
+				Handle = *FoundHandle;
 			}
 
 			if(!Handle.IsValid())
 			{
 				IdValues.Emplace(IdValue);
-				Handle.Handle = HandleType(SubsystemType, IdValues.Num());
-				IdValueToHandleMap.Emplace(IdValue, Handle.Handle);
+				Handle = HandleType(OnlineServicesType, IdValues.Num());
+				IdValueToHandleMap.Emplace(IdValue, Handle);
 			}
 		}
-
 		return Handle;
 	}
 
 	// Returns a copy as it's not thread safe to return a pointer/ref to an element of an array that can be relocated by another thread.
-	IdValueType GetIdValue(const HandleType Handle) const
+	IdValueType FindIdValue(const HandleType& Handle) const
 	{
-		if (Handle.Type == SubsystemType
-			&& Handle.IsValid()
-			&& IdValues.IsValidIndex(Handle.Handle - 1))
+		if (ValidateOnlineId(Handle))
 		{
 			FReadScopeLock ReadLock(Lock);
-			return IdValues[Handle.Handle - 1];
+			if (IdValues.IsValidIndex(Handle.GetHandle() - 1))
+			{
+				return IdValues[Handle.GetHandle() -1];
+			}
 		}
 		return IdValueType();
 	}
 
-	IdValueType GetIdValueChecked(const HandleType Handle) const
+	static inline bool ValidateOnlineId(const HandleType& Handle)
 	{
-		check(Handle.Type == SubsystemType
-			&& Handle.IsValid()
-			&& IdValues.IsValidIndex(Handle.Handle - 1));
-		return GetIdValue(Handle);
+		return ensure(Handle.GetOnlineServicesType() == OnlineServicesType) && Handle.IsValid();
 	}
 
 private:
 	mutable FRWLock Lock;
 
 	TArray<IdValueType> IdValues;
-	TMap<IdValueType, uint32> IdValueToHandleMap;
+	TMap<IdValueType, HandleType> IdValueToHandleMap;
 };
 
+template<typename IdValueType, EOnlineServices OnlineServicesType>
+using TOnlineBasicAccountIdRegistry = TOnlineBasicIdRegistry<OnlineIdHandleTags::FAccount, IdValueType, OnlineServicesType>;
+
 /* UE::Online */ }
 
diff --git a/Engine/Plugins/Online/Experimental/OnlineServicesEOS/OnlineServicesEOS.uplugin b/Engine/Plugins/Online/Experimental/OnlineServicesEOS/OnlineServicesEOS.uplugin
--- a/Engine/Plugins/Online/Experimental/OnlineServicesEOS/OnlineServicesEOS.uplugin
+++ b/Engine/Plugins/Online/Experimental/OnlineServicesEOS/OnlineServicesEOS.uplugin
@@ -4,7 +4,7 @@
 	"FriendlyName": "Online Services EOS",
 	"Version": 1,
 	"VersionName": "1.0",
-	"Description": "Online Services implementation for EOS",
+	"Description": "Online Services implementation for EOS Account and Game services.",
 	"Category": "Online Platform",
 	"CreatedBy": "Epic Games, Inc.",
 	"CreatedByURL": "http://epicgames.com",
@@ -34,11 +34,11 @@
 			"Enabled": true
 		},
 		{
-			"Name": "EOSShared",
+			"Name": "OnlineServicesEOSGS",
 			"Enabled": true
 		},
 		{
-			"Name": "SocketSubsystemEOS",
+			"Name": "EOSShared",
 			"Enabled": true
 		}
 	]
diff --git a/Engine/Plugins/Online/Experimental/OnlineServicesEOS/Source/OnlineServicesEOS.Build.cs b/Engine/Plugins/Online/Experimental/OnlineServicesEOS/Source/OnlineServicesEOS.Build.cs
--- a/Engine/Plugins/Online/Experimental/OnlineServicesEOS/Source/OnlineServicesEOS.Build.cs
+++ b/Engine/Plugins/Online/Experimental/OnlineServicesEOS/Source/OnlineServicesEOS.Build.cs
@@ -12,6 +12,7 @@ public class OnlineServicesEOS : ModuleRules
 				"Core",
 				"OnlineServicesInterface",
 				"OnlineServicesCommon",
+				"OnlineServicesEOSGS"
 			}
 		);
 
@@ -20,19 +21,8 @@ public class OnlineServicesEOS : ModuleRules
 				"CoreOnline",
 				"CoreUObject",
 				"EOSSDK",
-				"EOSShared"
+				"EOSShared",
 			}
 		);
-
-		if (Target.bCompileAgainstEngine)
-		{
-			PrivateDependencyModuleNames.AddRange(
-				new string[]
-				{
-					"SocketSubsystemEOS",
-					"Sockets"
-				}
-			);
-		}
 	}
 }
diff --git a/Engine/Plugins/Online/Experimental/OnlineServicesEOS/Source/Private/Online/AuthEOS.cpp b/Engine/Plugins/Online/Experimental/OnlineServicesEOS/Source/Private/Online/AuthEOS.cpp
--- a/Engine/Plugins/Online/Experimental/OnlineServicesEOS/Source/Private/Online/AuthEOS.cpp
+++ b/Engine/Plugins/Online/Experimental/OnlineServicesEOS/Source/Private/Online/AuthEOS.cpp
@@ -35,190 +35,19 @@ END_ONLINE_STRUCT_META()
 
 /* Meta*/ }
 
-static inline ELoginStatus ToELoginStatus(EOS_ELoginStatus InStatus)
-{
-	switch (InStatus)
-	{
-	case EOS_ELoginStatus::EOS_LS_NotLoggedIn:
-	{
-		return ELoginStatus::NotLoggedIn;
-	}
-	case EOS_ELoginStatus::EOS_LS_UsingLocalProfile:
-	{
-		return ELoginStatus::UsingLocalProfile;
-	}
-	case EOS_ELoginStatus::EOS_LS_LoggedIn:
-	{
-		return ELoginStatus::LoggedIn;
-	}
-	}
-	return ELoginStatus::NotLoggedIn;
-}
-
-// Copied from OSS EOS
-
-#define EOS_OSS_STRING_BUFFER_LENGTH 256
-// Chose arbitrarily since the SDK doesn't define it
-#define EOS_MAX_TOKEN_SIZE 4096
-
-struct FEOSAuthCredentials :
-	public EOS_Auth_Credentials
-{
-	FEOSAuthCredentials() :
-		EOS_Auth_Credentials()
-	{
-		ApiVersion = EOS_AUTH_CREDENTIALS_API_LATEST;
-		Id = IdAnsi;
-		Token = TokenAnsi;
-	}
-
-	FEOSAuthCredentials(const FEOSAuthCredentials& Other)
-	{
-		ApiVersion = Other.ApiVersion;
-		Id = IdAnsi;
-		Token = TokenAnsi;
-		Type = Other.Type;
-		SystemAuthCredentialsOptions = Other.SystemAuthCredentialsOptions;
-		ExternalType = Other.ExternalType;
-
-		FCStringAnsi::Strncpy(IdAnsi, Other.IdAnsi, EOS_OSS_STRING_BUFFER_LENGTH);
-		FCStringAnsi::Strncpy(TokenAnsi, Other.TokenAnsi, EOS_MAX_TOKEN_SIZE);
-	}
-
-	FEOSAuthCredentials(EOS_EExternalCredentialType InExternalType, const TArray<uint8>& InToken) :
-		EOS_Auth_Credentials()
-	{
-		ApiVersion = EOS_AUTH_CREDENTIALS_API_LATEST;
-		Type = EOS_ELoginCredentialType::EOS_LCT_ExternalAuth;
-		ExternalType = InExternalType;
-		Id = IdAnsi;
-		Token = TokenAnsi;
-
-		uint32_t InOutBufferLength = EOS_MAX_TOKEN_SIZE;
-		EOS_ByteArray_ToString(InToken.GetData(), InToken.Num(), TokenAnsi, &InOutBufferLength);
-	}
-
-	FEOSAuthCredentials& operator=(FEOSAuthCredentials& Other)
-	{
-		ApiVersion = Other.ApiVersion;
-		Type = Other.Type;
-		SystemAuthCredentialsOptions = Other.SystemAuthCredentialsOptions;
-		ExternalType = Other.ExternalType;
-
-		FCStringAnsi::Strncpy(IdAnsi, Other.IdAnsi, EOS_OSS_STRING_BUFFER_LENGTH);
-		FCStringAnsi::Strncpy(TokenAnsi, Other.TokenAnsi, EOS_MAX_TOKEN_SIZE);
-
-		return *this;
-	}
-
-	void SetToken(const FCredentialsToken& InToken)
-	{
-		if (InToken.IsType<TArray<uint8>>())
-		{
-			const TArray<uint8>& TokenData = InToken.Get<TArray<uint8>>();
-			uint32_t InOutBufferLength = EOS_MAX_TOKEN_SIZE;
-			EOS_ByteArray_ToString(TokenData.GetData(), TokenData.Num(), TokenAnsi, &InOutBufferLength);
-		}
-		else if (InToken.IsType<FString>())
-		{
-			const FString& TokenString = InToken.Get<FString>();
-			FCStringAnsi::Strncpy(TokenAnsi, TCHAR_TO_UTF8(*TokenString), EOS_MAX_TOKEN_SIZE);
-		}
-		else
-		{
-			UE_LOG(LogTemp, Warning, TEXT("SetToken cannot be set with an invalid credentials token parameter. Please ensure there is valid data in the credentials token."));
-		}
-	}
-
-	char IdAnsi[EOS_OSS_STRING_BUFFER_LENGTH];
-	char TokenAnsi[EOS_MAX_TOKEN_SIZE];
-};
-
-class FEOSConnectLoginCredentials : public EOS_Connect_Credentials
-{
-public:
-	FEOSConnectLoginCredentials()
-	{
-		ApiVersion = EOS_CONNECT_CREDENTIALS_API_LATEST;
-	}
-	virtual ~FEOSConnectLoginCredentials() {}
-};
-
-class FEOSConnectLoginCredentialsEOS : public FEOSConnectLoginCredentials
-{
-public:
-	FEOSConnectLoginCredentialsEOS()
-	{
-		Token = TokenAnsi;
-		TokenAnsi[0] = '\0';
-	}
-	void SetToken(const FCredentialsToken& InToken)
-	{
-		if (InToken.IsType<TArray<uint8>>())
-		{
-			const TArray<uint8>& TokenData = InToken.Get<TArray<uint8>>();
-			uint32_t InOutBufferLength = EOS_MAX_TOKEN_SIZE;
-			EOS_ByteArray_ToString(TokenData.GetData(), TokenData.Num(), TokenAnsi, &InOutBufferLength);
-		}
-		else if (InToken.IsType<FString>())
-		{
-			const FString& TokenString = InToken.Get<FString>();
-			FCStringAnsi::Strncpy(TokenAnsi, TCHAR_TO_UTF8(*TokenString), EOS_MAX_TOKEN_SIZE);
-		}
-		else
-		{
-			UE_LOG(LogTemp, Warning, TEXT("SetToken cannot be set with an invalid credentials token parameter. Please ensure there is valid data in the credentials token."));
-		}
-	}
-
-	char TokenAnsi[EOS_MAX_TOKEN_SIZE];
-};
-
-class FEOSConnectLoginCredentialsEAS
-	: public FEOSConnectLoginCredentials
-	, public FNoncopyable
-{
-public:
-	FEOSConnectLoginCredentialsEAS(EOS_Auth_Token* InEASToken)
-		: EASToken(InEASToken)
-	{
-		check(EASToken);
-		Type = EOS_EExternalCredentialType::EOS_ECT_EPIC;
-		Token = EASToken->AccessToken;
-	}
-	virtual ~FEOSConnectLoginCredentialsEAS()
-	{
-		if (EASToken)
-		{
-			EOS_Auth_Token_Release(EASToken);
-		}
-	}
-protected:
-	EOS_Auth_Token* EASToken = nullptr;
-};
-
-
-
-FAuthEOS::FAuthEOS(FOnlineServicesEOS& InServices, bool bInUseEAS)
-	: FAuthCommon(InServices)
-	, bUseEAS(bInUseEAS)
+FAuthEOS::FAuthEOS(FOnlineServicesEOS& InServices)
+	: Super(InServices)
 {
 }
 
 void FAuthEOS::Initialize()
 {
-	FAuthCommon::Initialize();
-
-	AuthHandle = EOS_Platform_GetAuthInterface(static_cast<FOnlineServicesEOS&>(GetServices()).GetEOSPlatformHandle());
-	check(AuthHandle != nullptr);
-
-	ConnectHandle = EOS_Platform_GetConnectInterface(static_cast<FOnlineServicesEOS&>(GetServices()).GetEOSPlatformHandle());
-	check(ConnectHandle != nullptr);
+	Super::Initialize();
 
 	// Register for login status changes
 	EOS_Auth_AddNotifyLoginStatusChangedOptions Options = { };
 	Options.ApiVersion = EOS_AUTH_ADDNOTIFYLOGINSTATUSCHANGED_API_LATEST;
-	NotifyLoginStatusChangedNotificationId = EOS_Auth_AddNotifyLoginStatusChanged(AuthHandle, &Options, this, [](const EOS_Auth_LoginStatusChangedCallbackInfo* Data)
+	NotifyEASLoginStatusChangedNotificationId = EOS_Auth_AddNotifyLoginStatusChanged(AuthHandle, &Options, this, [](const EOS_Auth_LoginStatusChangedCallbackInfo* Data)
 	{
 		FAuthEOS* This = reinterpret_cast<FAuthEOS*>(Data->ClientData);
 
@@ -228,124 +57,11 @@ void FAuthEOS::Initialize()
 		{
 			ELoginStatus PreviousStatus = ToELoginStatus(Data->PrevStatus);
 			ELoginStatus CurrentStatus = ToELoginStatus(Data->CurrentStatus);
-			This->OnEOSLoginStatusChanged(LocalUserId, PreviousStatus, CurrentStatus);
+			This->OnEASLoginStatusChanged(LocalUserId, PreviousStatus, CurrentStatus);
 		}
 	});
 }
 
-void FAuthEOS::PreShutdown()
-{
-}
-
-bool LexFromString(EOS_EAuthScopeFlags& OutEnum, const TCHAR* InString)
-{
-	if (FCString::Stricmp(InString, TEXT("BasicProfile")) == 0)
-	{
-		OutEnum = EOS_EAuthScopeFlags::EOS_AS_BasicProfile;
-	}
-	else if (FCString::Stricmp(InString, TEXT("FriendsList")) == 0)
-	{
-		OutEnum = EOS_EAuthScopeFlags::EOS_AS_FriendsList;
-	}
-	else if (FCString::Stricmp(InString, TEXT("Presence")) == 0)
-	{
-		OutEnum = EOS_EAuthScopeFlags::EOS_AS_Presence;
-	}
-	else if (FCString::Stricmp(InString, TEXT("FriendsManagement")) == 0)
-	{
-		OutEnum = EOS_EAuthScopeFlags::EOS_AS_FriendsManagement;
-	}
-	else if (FCString::Stricmp(InString, TEXT("Email")) == 0)
-	{
-		OutEnum = EOS_EAuthScopeFlags::EOS_AS_Email;
-	}
-	else if (FCString::Stricmp(InString, TEXT("NoFlags")) == 0 || FCString::Stricmp(InString, TEXT("None")) == 0)
-	{
-		OutEnum = EOS_EAuthScopeFlags::EOS_AS_NoFlags;
-	}
-	else
-	{
-		return false;
-	}
-	return true;
-}
-
-bool LexFromString(EOS_ELoginCredentialType& OutEnum, const TCHAR* const InString)
-{
-	if (FCString::Stricmp(InString, TEXT("ExchangeCode")) == 0)
-	{
-		OutEnum = EOS_ELoginCredentialType::EOS_LCT_ExchangeCode;
-	}
-	else if (FCString::Stricmp(InString, TEXT("PersistentAuth")) == 0)
-	{
-		OutEnum = EOS_ELoginCredentialType::EOS_LCT_PersistentAuth;
-	}
-	else if (FCString::Stricmp(InString, TEXT("Password")) == 0)
-	{
-		OutEnum = EOS_ELoginCredentialType::EOS_LCT_Password;
-	}
-	else if (FCString::Stricmp(InString, TEXT("Developer")) == 0)
-	{
-		OutEnum = EOS_ELoginCredentialType::EOS_LCT_Developer;
-	}
-	else if (FCString::Stricmp(InString, TEXT("RefreshToken")) == 0)
-	{
-		OutEnum = EOS_ELoginCredentialType::EOS_LCT_RefreshToken;
-	}
-	else if (FCString::Stricmp(InString, TEXT("AccountPortal")) == 0)
-	{
-		OutEnum = EOS_ELoginCredentialType::EOS_LCT_AccountPortal;
-	}
-	else if (FCString::Stricmp(InString, TEXT("ExternalAuth")) == 0)
-	{
-		OutEnum = EOS_ELoginCredentialType::EOS_LCT_ExternalAuth;
-	}
-	else
-	{
-		return false;
-	}
-	return true;
-}
-
-bool LexFromString(EOS_EExternalCredentialType& OutEnum, const TCHAR* InString)
-{
-	if (FCString::Stricmp(InString, TEXT("Steam")) == 0)
-	{
-		OutEnum = EOS_EExternalCredentialType::EOS_ECT_STEAM_APP_TICKET;
-	}
-	else if (FCString::Stricmp(InString, TEXT("PSN")) == 0)
-	{
-		OutEnum = EOS_EExternalCredentialType::EOS_ECT_PSN_ID_TOKEN;
-	}
-	else if (FCString::Stricmp(InString, TEXT("Xbox")) == 0)
-	{
-		OutEnum = EOS_EExternalCredentialType::EOS_ECT_XBL_XSTS_TOKEN;
-	}
-	else if (FCString::Stricmp(InString, TEXT("Nintendo")) == 0)
-	{
-		OutEnum = EOS_EExternalCredentialType::EOS_ECT_NINTENDO_ID_TOKEN;
-	}
-	else if (FCString::Stricmp(InString, TEXT("NSA")) == 0)
-	{
-		OutEnum = EOS_EExternalCredentialType::EOS_ECT_NINTENDO_NSA_ID_TOKEN;
-	}
-	else if (FCString::Stricmp(InString, TEXT("Apple")) == 0)
-	{
-		OutEnum = EOS_EExternalCredentialType::EOS_ECT_APPLE_ID_TOKEN;
-	}
-	else if (FCString::Stricmp(InString, TEXT("Google")) == 0)
-	{
-		OutEnum = EOS_EExternalCredentialType::EOS_ECT_GOOGLE_ID_TOKEN;
-	}
-	else
-	{
-		// Unknown means OpenID
-		OutEnum = EOS_EExternalCredentialType::EOS_ECT_OPENID_ACCESS_TOKEN;
-	}
-
-	return true;
-}
-
 TOnlineAsyncOpHandle<FAuthLogin> FAuthEOS::Login(FAuthLogin::Params&& Params)
 {
 	// Is this auto-login?
@@ -372,22 +88,14 @@ TOnlineAsyncOpHandle<FAuthLogin> FAuthEOS::Login(FAuthLogin::Params&& Params)
 		return Op->GetHandle();
 	}
 
-	TOnlineChainableAsyncOp<FAuthLogin, TSharedPtr<FEOSConnectLoginCredentials>> ConnectLoginOp = [this, &Op]()
-	{
-		if (bUseEAS)
-		{
-			return LoginEAS(*Op);
-		}
-		return Op->Then([this](TOnlineAsyncOp<FAuthLogin>& InAsyncOp)
-		{
-			return MakeConnectLoginCredentials(InAsyncOp);
-		});
-	}();
-	// Check if the above steps completed (failed) the operation
+	TOnlineChainableAsyncOp<FAuthLogin, TSharedPtr<FEOSConnectLoginCredentials>> ConnectLoginOp = LoginEAS(*Op);
+
+	// Check if LoginEAS completed (failed) the operation
 	if (ConnectLoginOp.GetOwningOperation().IsComplete())
 	{
 		return ConnectLoginOp.GetOwningOperation().GetHandle();
 	}
+
 	ConnectLoginOp.Then([this](TOnlineAsyncOp<FAuthLogin>& InAsyncOp, TSharedPtr<FEOSConnectLoginCredentials>&& InConnectLoginCredentials)
 	{
 		EOS_Connect_LoginOptions ConnectLoginOptions = { };
@@ -520,12 +228,12 @@ TOnlineChainableAsyncOp<FAuthLogin, TSharedPtr<FEOSConnectLoginCredentials>> FAu
 			FCStringAnsi::Strncpy(Credentials.TokenAnsi, TCHAR_TO_UTF8(*Params.CredentialsToken.Get<FString>()), EOS_MAX_TOKEN_SIZE);
 			break;
 		case EOS_ELoginCredentialType::EOS_LCT_Password:
-			FCStringAnsi::Strncpy(Credentials.IdAnsi, TCHAR_TO_UTF8(*Params.CredentialsId), EOS_OSS_STRING_BUFFER_LENGTH);
+			FCStringAnsi::Strncpy(Credentials.IdAnsi, TCHAR_TO_UTF8(*Params.CredentialsId), EOS_STRING_BUFFER_LENGTH);
 			FCStringAnsi::Strncpy(Credentials.TokenAnsi, TCHAR_TO_UTF8(*Params.CredentialsToken.Get<FString>()), EOS_MAX_TOKEN_SIZE);
 			break;
 		case EOS_ELoginCredentialType::EOS_LCT_Developer:
 			// This is auth via the EOS auth tool
-			FCStringAnsi::Strncpy(Credentials.IdAnsi, TCHAR_TO_UTF8(*Params.CredentialsId), EOS_OSS_STRING_BUFFER_LENGTH);
+			FCStringAnsi::Strncpy(Credentials.IdAnsi, TCHAR_TO_UTF8(*Params.CredentialsId), EOS_STRING_BUFFER_LENGTH);
 			FCStringAnsi::Strncpy(Credentials.TokenAnsi, TCHAR_TO_UTF8(*Params.CredentialsToken.Get<FString>()), EOS_MAX_TOKEN_SIZE);
 			break;
 		case EOS_ELoginCredentialType::EOS_LCT_AccountPortal:
@@ -636,16 +344,6 @@ TOnlineChainableAsyncOp<FAuthLogin, TSharedPtr<FEOSConnectLoginCredentials>> FAu
 	});
 }
 
-TSharedPtr<FEOSConnectLoginCredentials> FAuthEOS::MakeConnectLoginCredentials(TOnlineAsyncOp<FAuthLogin>& InAsyncOp)
-{
-	const FAuthLogin::Params& LoginParams = InAsyncOp.GetParams();
-	TSharedRef<FEOSConnectLoginCredentials> Credentials = MakeShared<FEOSConnectLoginCredentialsEOS>();
-	FEOSConnectLoginCredentialsEOS& CredentialsEOS = static_cast<FEOSConnectLoginCredentialsEOS&>(*Credentials);
-	LexFromString(CredentialsEOS.Type, *LoginParams.CredentialsType);
-	CredentialsEOS.SetToken(LoginParams.CredentialsToken);
-	return Credentials;
-}
-
 void FAuthEOS::ProcessSuccessfulLogin(TOnlineAsyncOp<FAuthLogin>& InAsyncOp)
 {
 	const EOS_EpicAccountId* EpicAccountId = InAsyncOp.Data.Get<EOS_EpicAccountId>(TEXT("EpicAccountId"));
@@ -686,7 +384,7 @@ void FAuthEOS::ProcessSuccessfulLogin(TOnlineAsyncOp<FAuthLogin>& InAsyncOp)
 	FAuthLogin::Result Result = { AccountInfo };
 	InAsyncOp.SetResult(MoveTemp(Result));
 
-	// When a user logs in, OnEOSLoginStatusChanged can not trigger (if it's that user's first login) or trigger before we add relevant information to AccountInfos, so we trigger the status change event here 
+	// When a user logs in, OnEASLoginStatusChanged can not trigger (if it's that user's first login) or trigger before we add relevant information to AccountInfos, so we trigger the status change event here 
 	FLoginStatusChanged EventParameters;
 	EventParameters.LocalUserId = LocalUserId;
 	EventParameters.PreviousStatus = ELoginStatus::NotLoggedIn;
@@ -694,6 +392,32 @@ void FAuthEOS::ProcessSuccessfulLogin(TOnlineAsyncOp<FAuthLogin>& InAsyncOp)
 	OnLoginStatusChangedEvent.Broadcast(EventParameters);
 }
 
+void FAuthEOS::OnEASLoginStatusChanged(FOnlineAccountIdHandle LocalUserId, ELoginStatus PreviousStatus, ELoginStatus CurrentStatus)
+{
+	UE_LOG(LogTemp, Warning, TEXT("OnEASLoginStatusChanged: [%s] [%s]->[%s]"), *ToLogString(LocalUserId), LexToString(PreviousStatus), LexToString(CurrentStatus));
+	if (TSharedRef<FAccountInfoEOS>* AccountInfoPtr = AccountInfos.Find(LocalUserId))
+	{
+		FAccountInfoEOS& AccountInfo = **AccountInfoPtr;
+		if (AccountInfo.LoginStatus != CurrentStatus)
+		{
+			FLoginStatusChanged EventParameters;
+			EventParameters.LocalUserId = LocalUserId;
+			EventParameters.PreviousStatus = AccountInfo.LoginStatus;
+			EventParameters.CurrentStatus = CurrentStatus;
+
+			AccountInfo.LoginStatus = CurrentStatus;
+
+			if (CurrentStatus == ELoginStatus::NotLoggedIn)
+			{
+				// Remove user
+				AccountInfos.Remove(LocalUserId); // Invalidates AccountInfo
+			}
+
+			OnLoginStatusChangedEvent.Broadcast(EventParameters);
+		}
+	}
+}
+
 TOnlineAsyncOpHandle<FAuthLogout> FAuthEOS::Logout(FAuthLogout::Params&& Params)
 {
 	const FOnlineAccountIdHandle LocalUserId = Params.LocalUserId;
@@ -748,7 +472,7 @@ TOnlineAsyncOpHandle<FAuthLogout> FAuthEOS::Logout(FAuthLogout::Params&& Params)
 			// Success
 			InAsyncOp.SetResult(FAuthLogout::Result());
 
-			// OnLoginStatusChanged will be triggered by OnEOSLoginStatusChanged
+			// OnLoginStatusChanged will be triggered by OnEASLoginStatusChanged
 		}
 		else
 		{
@@ -761,80 +485,6 @@ TOnlineAsyncOpHandle<FAuthLogout> FAuthEOS::Logout(FAuthLogout::Params&& Params)
 	return Op->GetHandle();
 }
 
-TOnlineResult<FAuthGetAccountByPlatformUserId> FAuthEOS::GetAccountByPlatformUserId(FAuthGetAccountByPlatformUserId::Params&& Params)
-{
-	TResult<FOnlineAccountIdHandle, FOnlineError> PlatformUserIdResult = GetAccountIdByPlatformUserId(Params.PlatformUserId);
-	if (PlatformUserIdResult.IsOk())
-	{
-		FAuthGetAccountByPlatformUserId::Result Result = { AccountInfos.FindChecked(PlatformUserIdResult.GetOkValue()) };
-		return TOnlineResult<FAuthGetAccountByPlatformUserId>(Result);
-	}
-	else
-	{
-		return TOnlineResult<FAuthGetAccountByPlatformUserId>(PlatformUserIdResult.GetErrorValue());
-	}
-}
-
-TOnlineResult<FAuthGetAccountByAccountId> FAuthEOS::GetAccountByAccountId(FAuthGetAccountByAccountId::Params&& Params)
-{
-	if (TSharedRef<FAccountInfoEOS>* const FoundAccount = AccountInfos.Find(Params.LocalUserId))
-	{
-		FAuthGetAccountByAccountId::Result Result = { *FoundAccount };
-		return TOnlineResult<FAuthGetAccountByAccountId>(Result);
-	}
-	else
-	{
-		// TODO: proper error
-		return TOnlineResult<FAuthGetAccountByAccountId>(Errors::Unknown());
-	}
-}
-
-bool FAuthEOS::IsLoggedIn(const FOnlineAccountIdHandle& AccountId) const
-{
-	// TODO:  More logic?
-	return AccountInfos.Contains(AccountId);
-}
-
-TResult<FOnlineAccountIdHandle, FOnlineError> FAuthEOS::GetAccountIdByPlatformUserId(FPlatformUserId PlatformUserId) const
-{
-	for (const TPair<FOnlineAccountIdHandle, TSharedRef<FAccountInfoEOS>>& AccountPair : AccountInfos)
-	{
-		if (AccountPair.Value->PlatformUserId == PlatformUserId)
-		{
-			TResult<FOnlineAccountIdHandle, FOnlineError> Result(AccountPair.Key);
-			return Result;
-		}
-	}
-	TResult<FOnlineAccountIdHandle, FOnlineError> Result(Errors::Unknown()); // TODO: error code
-	return Result;
-}
-
-void FAuthEOS::OnEOSLoginStatusChanged(FOnlineAccountIdHandle LocalUserId, ELoginStatus PreviousStatus, ELoginStatus CurrentStatus)
-{
-	UE_LOG(LogTemp, Warning, TEXT("OnEOSLoginStatusChanged: [%s] [%s]->[%s]"), *ToLogString(LocalUserId), LexToString(PreviousStatus), LexToString(CurrentStatus));
-	if (TSharedRef<FAccountInfoEOS>* AccountInfoPtr = AccountInfos.Find(LocalUserId))
-	{
-		FAccountInfoEOS& AccountInfo = **AccountInfoPtr;
-		if (AccountInfo.LoginStatus != CurrentStatus)
-		{
-			FLoginStatusChanged EventParameters;
-			EventParameters.LocalUserId = LocalUserId;
-			EventParameters.PreviousStatus = AccountInfo.LoginStatus;
-			EventParameters.CurrentStatus = CurrentStatus;
-
-			AccountInfo.LoginStatus = CurrentStatus;
-
-			if (CurrentStatus == ELoginStatus::NotLoggedIn)
-			{
-				// Remove user
-				AccountInfos.Remove(LocalUserId); // Invalidates AccountInfo
-			}
-
-			OnLoginStatusChangedEvent.Broadcast(EventParameters);
-		}
-	}
-}
-
 template <typename IdType>
 TFuture<FOnlineAccountIdHandle> ResolveAccountIdImpl(FAuthEOS& AuthEOS, const FOnlineAccountIdHandle& LocalUserId, const IdType InId)
 {
@@ -1091,4 +741,10 @@ TFunction<TFuture<TArray<FOnlineAccountIdHandle>>(FOnlineAsyncOp& InAsyncOp, con
 	return ResolveIdsFnImpl<EOS_ProductUserId>(this);
 }
 
+FOnlineAccountIdHandle FAuthEOS::CreateAccountId(const EOS_EpicAccountId EpicAccountId, const EOS_ProductUserId ProductUserId)
+{
+	return FOnlineAccountIdRegistryEOS::Get().FindOrAddAccountId(EpicAccountId, ProductUserId);
+}
+
+
 /* UE::Online */ }
diff --git a/Engine/Plugins/Online/Experimental/OnlineServicesEOS/Source/Private/Online/LobbiesEOS.cpp b/Engine/Plugins/Online/Experimental/OnlineServicesEOS/Source/Private/Online/LobbiesEOS.cpp
deleted file mode 100644
diff --git a/Engine/Plugins/Online/Experimental/OnlineServicesEOS/Source/Private/Online/LobbiesEOSTypes.cpp b/Engine/Plugins/Online/Experimental/OnlineServicesEOS/Source/Private/Online/LobbiesEOSTypes.cpp
deleted file mode 100644
diff --git a/Engine/Plugins/Online/Experimental/OnlineServicesEOS/Source/Private/Online/LobbiesEOSTypes.h b/Engine/Plugins/Online/Experimental/OnlineServicesEOS/Source/Private/Online/LobbiesEOSTypes.h
deleted file mode 100644
diff --git a/Engine/Plugins/Online/Experimental/OnlineServicesEOS/Source/Private/Online/OnlineErrorEOS.cpp b/Engine/Plugins/Online/Experimental/OnlineServicesEOS/Source/Private/Online/OnlineErrorEOS.cpp
deleted file mode 100644
diff --git a/Engine/Plugins/Online/Experimental/OnlineServicesEOS/Source/Private/Online/OnlineErrorEOS.h b/Engine/Plugins/Online/Experimental/OnlineServicesEOS/Source/Private/Online/OnlineErrorEOS.h
deleted file mode 100644
diff --git a/Engine/Plugins/Online/Experimental/OnlineServicesEOS/Source/Private/Online/OnlineIdEOS.cpp b/Engine/Plugins/Online/Experimental/OnlineServicesEOS/Source/Private/Online/OnlineIdEOS.cpp
--- a/Engine/Plugins/Online/Experimental/OnlineServicesEOS/Source/Private/Online/OnlineIdEOS.cpp
+++ b/Engine/Plugins/Online/Experimental/OnlineServicesEOS/Source/Private/Online/OnlineIdEOS.cpp
@@ -7,24 +7,118 @@
 #include "String/BytesToHex.h"
 #include "Online/OnlineServicesEOSTypes.h"
 
-#include "eos_connect_types.h"
-#include "eos_connect.h"
-
 namespace UE::Online {
 
-FOnlineAccountIdRegistryEOS& FOnlineAccountIdRegistryEOS::Get()
+FString FOnlineAccountIdRegistryEOS::ToLogString(const FOnlineAccountIdHandle& Handle) const
 {
-	static FOnlineAccountIdRegistryEOS Instance;
-	return Instance;
+	FString Result;
+	if (ValidateOnlineId(Handle))
+	{
+		const FOnlineAccountIdDataEOS& IdData = GetAccountIdData(Handle);
+		Result = FString::Printf(TEXT("EAS=[%s] EOS=[%s]"), *LexToString(IdData.EpicAccountId), *LexToString(IdData.ProductUserId));
+	}
+	else
+	{
+		check(!Handle.IsValid()); // Check we haven't been passed a valid handle for a different EOnlineServices.
+		Result = TEXT("Invalid");
+	}
+	return Result;
 }
 
-FOnlineAccountIdHandle FOnlineAccountIdRegistryEOS::Find(const EOS_EpicAccountId EpicAccountId) const
+enum class EEOSAccountIdElements : uint8
 {
+	None = 0,
+	EAS = 1 << 0,
+	EOS = 1 << 1
+};
+ENUM_CLASS_FLAGS(EEOSAccountIdElements);
+const uint8 OnlineIdEOSUtf8BufferLength = 32;
+const uint8 OnlineIdEOSHexBufferLength = 16;
+
+TArray<uint8> FOnlineAccountIdRegistryEOS::ToReplicationData(const FOnlineAccountIdHandle& Handle) const
+{
+	TArray<uint8> ReplicationData;
+	if (ValidateOnlineId(Handle))
+	{
+		const FOnlineAccountIdDataEOS& IdData = GetAccountIdData(Handle);
+		EEOSAccountIdElements Elements = EEOSAccountIdElements::None;
+
+		char EasBuffer[EOS_EPICACCOUNTID_MAX_LENGTH + 1] = { 0 };
+		if (EOS_EpicAccountId_IsValid(IdData.EpicAccountId))
+		{
+			int32 EasBufferLength = sizeof(EasBuffer);
+			const EOS_EResult EosResult = EOS_EpicAccountId_ToString(IdData.EpicAccountId, EasBuffer, &EasBufferLength);
+			if (ensure(EosResult == EOS_EResult::EOS_Success))
+			{
+				check(EasBufferLength - 1 == OnlineIdEOSUtf8BufferLength);
+				Elements |= EEOSAccountIdElements::EAS;
+			}
+		}
+
+		char EosBuffer[EOS_PRODUCTUSERID_MAX_LENGTH + 1] = { 0 };
+		if (EOS_ProductUserId_IsValid(IdData.ProductUserId))
+		{
+			int32 EosBufferLength = sizeof(EosBuffer);
+			const EOS_EResult EosResult = EOS_ProductUserId_ToString(IdData.ProductUserId, EosBuffer, &EosBufferLength);
+			if (ensure(EosResult == EOS_EResult::EOS_Success))
+			{
+				check(EosBufferLength - 1 == OnlineIdEOSUtf8BufferLength);
+				Elements |= EEOSAccountIdElements::EOS;
+			}
+		}
+
+		const int EasHexBufferLength = EnumHasAnyFlags(Elements, EEOSAccountIdElements::EAS) ? OnlineIdEOSHexBufferLength : 0;
+		const int EosHexBufferLength = EnumHasAnyFlags(Elements, EEOSAccountIdElements::EOS) ? OnlineIdEOSHexBufferLength : 0;
+		ReplicationData.SetNumUninitialized(1 + EasHexBufferLength + EosHexBufferLength);
+		ReplicationData[0] = uint8(Elements);
+		if (EasHexBufferLength > 0)
+		{
+			UE::String::HexToBytes(FUtf8StringView(EasBuffer, OnlineIdEOSUtf8BufferLength), ReplicationData.GetData() + 1);
+		}
+		if (EosHexBufferLength > 0)
+		{
+			UE::String::HexToBytes(FUtf8StringView(EosBuffer, OnlineIdEOSUtf8BufferLength), ReplicationData.GetData() + 1 + EasHexBufferLength);
+		}
+	}
+	else
+	{
+		check(!Handle.IsValid()); // Check we haven't been passed a valid handle for a different EOnlineServices.
+	}
+	return ReplicationData;
+}
+
+FOnlineAccountIdHandle FOnlineAccountIdRegistryEOS::FromReplicationData(const TArray<uint8>& ReplicationData)
+{
+	const EEOSAccountIdElements Elements = EEOSAccountIdElements(ReplicationData[0]);
+	const int EasHexBufferLength = EnumHasAnyFlags(Elements, EEOSAccountIdElements::EAS) ? OnlineIdEOSHexBufferLength : 0;
+	const int EosHexBufferLength = EnumHasAnyFlags(Elements, EEOSAccountIdElements::EOS) ? OnlineIdEOSHexBufferLength : 0;
+
+	EOS_EpicAccountId EpicAccountId = nullptr;
+	if (EasHexBufferLength > 0)
+	{
+		char EasBuffer[EOS_EPICACCOUNTID_MAX_LENGTH + 1] = { 0 };
+		UE::String::BytesToHex(TConstArrayView<uint8>(ReplicationData.GetData() + 1, EasHexBufferLength), EasBuffer);
+		EpicAccountId = EOS_EpicAccountId_FromString(EasBuffer);
+	}
+
+	EOS_ProductUserId ProductUserId = nullptr;
+	if (EosHexBufferLength > 0)
+	{
+		char EosBuffer[EOS_EPICACCOUNTID_MAX_LENGTH + 1] = { 0 };
+		UE::String::BytesToHex(TConstArrayView<uint8>(ReplicationData.GetData() + 1 + EasHexBufferLength, EosHexBufferLength), EosBuffer);
+		ProductUserId = EOS_ProductUserId_FromString(EosBuffer);
+	}
+
+	return FindOrAddAccountId(EpicAccountId, ProductUserId);
+}
+
+FOnlineAccountIdHandle FOnlineAccountIdRegistryEOS::FindAccountId(const EOS_ProductUserId ProductUserId) const
+{
 	FOnlineAccountIdHandle Handle;
-	if (EOS_EpicAccountId_IsValid(EpicAccountId) == EOS_TRUE)
+	if (EOS_ProductUserId_IsValid(ProductUserId) == EOS_TRUE)
 	{
 		const FReadScopeLock ReadLock(Lock);
-		if (const FOnlineAccountIdHandle* Found = EpicAccountIdToHandle.Find(EpicAccountId))
+		if (const FOnlineAccountIdHandle* Found = ProductUserIdToHandle.Find(ProductUserId))
 		{
 			Handle = *Found;
 		}
@@ -32,13 +126,18 @@
 	return Handle;
 }
 
-FOnlineAccountIdHandle FOnlineAccountIdRegistryEOS::Find(const EOS_ProductUserId ProductUserId) const
+EOS_ProductUserId FOnlineAccountIdRegistryEOS::GetProductUserId(const FOnlineAccountIdHandle& Handle) const
 {
+	return GetAccountIdData(Handle).ProductUserId;
+}
+
+FOnlineAccountIdHandle FOnlineAccountIdRegistryEOS::FindAccountId(const EOS_EpicAccountId EpicAccountId) const
+{
 	FOnlineAccountIdHandle Handle;
-	if (EOS_ProductUserId_IsValid(ProductUserId) == EOS_TRUE)
+	if (EOS_EpicAccountId_IsValid(EpicAccountId) == EOS_TRUE)
 	{
 		const FReadScopeLock ReadLock(Lock);
-		if (const FOnlineAccountIdHandle* Found = ProductUserIdToHandle.Find(ProductUserId))
+		if (const FOnlineAccountIdHandle* Found = EpicAccountIdToHandle.Find(EpicAccountId))
 		{
 			Handle = *Found;
 		}
@@ -46,7 +145,23 @@
 	return Handle;
 }
 
-FOnlineAccountIdHandle FOnlineAccountIdRegistryEOS::Create(const EOS_EpicAccountId InEpicAccountId, const EOS_ProductUserId InProductUserId)
+FOnlineAccountIdDataEOS FOnlineAccountIdRegistryEOS::GetAccountIdData(const FOnlineAccountIdHandle& Handle) const
 {
+	if (ValidateOnlineId(Handle))
+	{
+		FReadScopeLock ScopeLock(Lock);
+		return AccountIdData[Handle.GetHandle() - 1];
+	}
+	return FOnlineAccountIdDataEOS();
+}
+
+FOnlineAccountIdRegistryEOS& FOnlineAccountIdRegistryEOS::Get()
+{
+	static FOnlineAccountIdRegistryEOS Instance;
+	return Instance;
+}
+
+FOnlineAccountIdHandle FOnlineAccountIdRegistryEOS::FindOrAddAccountId(const EOS_EpicAccountId InEpicAccountId, const EOS_ProductUserId InProductUserId)
+{
 	FOnlineAccountIdHandle Result;
 	const bool bInEpicAccountIdValid = EOS_EpicAccountId_IsValid(InEpicAccountId) == EOS_TRUE;
@@ -138,122 +253,9 @@
 	return Result;
 }
 
-FOnlineAccountIdDataEOS FOnlineAccountIdRegistryEOS::GetIdData(const FOnlineAccountIdHandle& Handle) const
-{
-	if (ValidateOnlineId(Handle))
-	{
-		FReadScopeLock ScopeLock(Lock);
-		return AccountIdData[Handle.GetHandle() - 1];
-	}
-	return FOnlineAccountIdDataEOS();
-}
-
-FString FOnlineAccountIdRegistryEOS::ToLogString(const FOnlineAccountIdHandle& Handle) const
-{
-	FString Result;
-	if (ValidateOnlineId(Handle))
-	{
-		const FOnlineAccountIdDataEOS& IdData = GetIdData(Handle);
-		Result = FString::Printf(TEXT("EAS=[%s] EOS=[%s]"), *LexToString(IdData.EpicAccountId), *LexToString(IdData.ProductUserId));
-	}
-	else
-	{
-		check(!Handle.IsValid()); // Check we haven't been passed a valid handle for a different EOnlineServices.
-		Result = TEXT("Invalid");
-	}
-	return Result;
-}
-
-enum class EEOSAccountIdElements : uint8
-{	
-	None = 0,
-	EAS = 1 << 0,
-	EOS = 1 << 1
-};
-ENUM_CLASS_FLAGS(EEOSAccountIdElements);
-const uint8 OnlineIdEOSUtf8BufferLength = 32;
-const uint8 OnlineIdEOSHexBufferLength = 16;
-
-TArray<uint8> FOnlineAccountIdRegistryEOS::ToReplicationData(const FOnlineAccountIdHandle& Handle) const
-{
-	TArray<uint8> ReplicationData;
-	if (ValidateOnlineId(Handle))
-	{
-		const FOnlineAccountIdDataEOS& IdData = GetIdData(Handle);
-		EEOSAccountIdElements Elements = EEOSAccountIdElements::None;
-
-		char EasBuffer[EOS_EPICACCOUNTID_MAX_LENGTH + 1] = { 0 };
-		if (EOS_EpicAccountId_IsValid(IdData.EpicAccountId))
-		{
-			int32 EasBufferLength = sizeof(EasBuffer);
-			const EOS_EResult EosResult = EOS_EpicAccountId_ToString(IdData.EpicAccountId, EasBuffer, &EasBufferLength);
-			if (ensure(EosResult == EOS_EResult::EOS_Success))
-			{
-				check(EasBufferLength-1 == OnlineIdEOSUtf8BufferLength);
-				Elements |= EEOSAccountIdElements::EAS;
-			}
-		}
-
-		char EosBuffer[EOS_PRODUCTUSERID_MAX_LENGTH + 1] = { 0 };
-		if (EOS_ProductUserId_IsValid(IdData.ProductUserId))
-		{
-			int32 EosBufferLength = sizeof(EosBuffer);
-			const EOS_EResult EosResult = EOS_ProductUserId_ToString(IdData.ProductUserId, EosBuffer, &EosBufferLength);
-			if (ensure(EosResult == EOS_EResult::EOS_Success))
-			{
-				check(EosBufferLength-1 == OnlineIdEOSUtf8BufferLength);
-				Elements |= EEOSAccountIdElements::EOS;
-			}
-		}
-
-		const int EasHexBufferLength = EnumHasAnyFlags(Elements, EEOSAccountIdElements::EAS) ? OnlineIdEOSHexBufferLength : 0;
-		const int EosHexBufferLength = EnumHasAnyFlags(Elements, EEOSAccountIdElements::EOS) ? OnlineIdEOSHexBufferLength : 0;
-		ReplicationData.SetNumUninitialized(1 + EasHexBufferLength + EosHexBufferLength);
-		ReplicationData[0] = uint8(Elements);
-		if (EasHexBufferLength > 0)
-		{
-			UE::String::HexToBytes(FUtf8StringView(EasBuffer, OnlineIdEOSUtf8BufferLength), ReplicationData.GetData() + 1);
-		}
-		if (EosHexBufferLength > 0)
-		{
-			UE::String::HexToBytes(FUtf8StringView(EosBuffer, OnlineIdEOSUtf8BufferLength), ReplicationData.GetData() + 1 + EasHexBufferLength);
-		}
-	}
-	else
-	{
-		check(!Handle.IsValid()); // Check we haven't been passed a valid handle for a different EOnlineServices.
-	}
-	return ReplicationData;
-}
-
-FOnlineAccountIdHandle FOnlineAccountIdRegistryEOS::FromReplicationData(const TArray<uint8>& ReplicationData)
-{
-	const EEOSAccountIdElements Elements = EEOSAccountIdElements(ReplicationData[0]);
-	const int EasHexBufferLength = EnumHasAnyFlags(Elements, EEOSAccountIdElements::EAS) ? OnlineIdEOSHexBufferLength : 0;
-	const int EosHexBufferLength = EnumHasAnyFlags(Elements, EEOSAccountIdElements::EOS) ? OnlineIdEOSHexBufferLength : 0;
-
-	EOS_EpicAccountId EpicAccountId = nullptr;
-	if (EasHexBufferLength > 0)
-	{
-		char EasBuffer[EOS_EPICACCOUNTID_MAX_LENGTH + 1] = { 0 };
-		UE::String::BytesToHex(TConstArrayView<uint8>(ReplicationData.GetData() + 1, EasHexBufferLength), EasBuffer);
-		EpicAccountId = EOS_EpicAccountId_FromString(EasBuffer);
-	}
-
-	EOS_ProductUserId ProductUserId = nullptr;
-	if (EosHexBufferLength > 0)
-	{
-		char EosBuffer[EOS_EPICACCOUNTID_MAX_LENGTH + 1] = { 0 };
-		UE::String::BytesToHex(TConstArrayView<uint8>(ReplicationData.GetData() + 1 + EasHexBufferLength, EosHexBufferLength), EosBuffer);
-		ProductUserId = EOS_ProductUserId_FromString(EosBuffer);
-	}
-	
-	return Create(EpicAccountId, ProductUserId);
-}
-
 EOS_EpicAccountId GetEpicAccountId(const FOnlineAccountIdHandle& Handle)
 {
-	return FOnlineAccountIdRegistryEOS::Get().GetIdData(Handle).EpicAccountId;
+	return FOnlineAccountIdRegistryEOS::Get().GetAccountIdData(Handle).EpicAccountId;
 }
 
 EOS_EpicAccountId GetEpicAccountIdChecked(const FOnlineAccountIdHandle& Handle)
@@ -263,26 +265,9 @@
 	return EpicAccountId;
 }
 
-EOS_ProductUserId GetProductUserId(const FOnlineAccountIdHandle& Handle)
-{
-	return FOnlineAccountIdRegistryEOS::Get().GetIdData(Handle).ProductUserId;
-}
-
-EOS_ProductUserId GetProductUserIdChecked(const FOnlineAccountIdHandle& Handle)
-{
-	EOS_ProductUserId ProductUserId = GetProductUserId(Handle);
-	check(EOS_ProductUserId_IsValid(ProductUserId) == EOS_TRUE);
-	return ProductUserId;
-}
-
 FOnlineAccountIdHandle FindAccountId(const EOS_EpicAccountId EpicAccountId)
 {
-	return FOnlineAccountIdRegistryEOS::Get().Find(EpicAccountId);
+	return FOnlineAccountIdRegistryEOS::Get().FindAccountId(EpicAccountId);
 }
 
-FOnlineAccountIdHandle FindAccountId(const EOS_ProductUserId ProductUserId)
-{
-	return FOnlineAccountIdRegistryEOS::Get().Find(ProductUserId);
-}
-
 FOnlineAccountIdHandle FindAccountIdChecked(const EOS_EpicAccountId EpicAccountId)
@@ -292,16 +277,4 @@
 	return Result;
 }
 
-FOnlineAccountIdHandle FindAccountIdChecked(const EOS_ProductUserId ProductUserId)
-{
-	FOnlineAccountIdHandle Result = FindAccountId(ProductUserId);
-	check(Result.IsValid());
-	return Result;
-}
-
-FOnlineAccountIdHandle CreateAccountId(const EOS_EpicAccountId EpicAccountId, const EOS_ProductUserId ProductUserId)
-{
-	return FOnlineAccountIdRegistryEOS::Get().Create(EpicAccountId, ProductUserId);
-}
-
 /* UE::Online */ }
diff --git a/Engine/Plugins/Online/Experimental/OnlineServicesEOS/Source/Private/Online/OnlineIdEOS.h b/Engine/Plugins/Online/Experimental/OnlineServicesEOS/Source/Private/Online/OnlineIdEOS.h
--- a/Engine/Plugins/Online/Experimental/OnlineServicesEOS/Source/Private/Online/OnlineIdEOS.h
+++ b/Engine/Plugins/Online/Experimental/OnlineServicesEOS/Source/Private/Online/OnlineIdEOS.h
@@ -2,16 +2,9 @@
 
 #pragma once
 
-#include "Async/Future.h"
-#include "Delegates/Delegate.h"
-#include "Online/CoreOnline.h"
-#include "Online/OnlineIdCommon.h"
+#include "Online/OnlineIdEOSGS.h"
 
-#if defined(EOS_PLATFORM_BASE_FILE_NAME)
-#include EOS_PLATFORM_BASE_FILE_NAME
-#endif
-#include "eos_common.h"
-#include "eos_connect_types.h"
+class FAccountIdReplicationTest;
 
 namespace UE::Online {
 
@@ -23,19 +16,12 @@ public:
 };
 
 /**
- * Account id registry specifically for EOS id's which are segmented.
+ * Account id registry specifically for EOS id's which are an EOS_EpicAccountId/EOS_ProductUserId pair.
  */
-class FOnlineAccountIdRegistryEOS : public IOnlineAccountIdRegistry
+class FOnlineAccountIdRegistryEOS : public IOnlineAccountIdRegistryEOSGS
 {
 public:
-	static FOnlineAccountIdRegistryEOS& Get();
-
-	FOnlineAccountIdHandle Find(const EOS_EpicAccountId EpicAccountId) const;
-	FOnlineAccountIdHandle Find(const EOS_ProductUserId ProductUserId) const;
-	FOnlineAccountIdHandle Create(const EOS_EpicAccountId EpicAccountId, const EOS_ProductUserId ProductUserId);
-
-	// Return copies as it is not thread safe to return pointers/references to array elements, in case the array is grown+relocated on another thread.
-	FOnlineAccountIdDataEOS GetIdData(const FOnlineAccountIdHandle& Handle) const;
+	virtual ~FOnlineAccountIdRegistryEOS() = default;
 
 	// Begin IOnlineAccountIdRegistry
 	virtual FString ToLogString(const FOnlineAccountIdHandle& Handle) const override;
@@ -43,34 +29,35 @@ public:
 	virtual FOnlineAccountIdHandle FromReplicationData(const TArray<uint8>& ReplicationString) override;
 	// End IOnlineAccountIdRegistry
 
-	virtual ~FOnlineAccountIdRegistryEOS() = default;
+	// Begin IOnlineAccountRegistryEOSGS
+	virtual FOnlineAccountIdHandle FindAccountId(const EOS_ProductUserId ProductUserId) const override;
+	virtual EOS_ProductUserId GetProductUserId(const FOnlineAccountIdHandle& Handle) const override;
+	// End IOnlineAccountRegistryEOSGS
+
+	FOnlineAccountIdHandle FindAccountId(const EOS_EpicAccountId EpicAccountId) const;
+
+	// Return copies as it is not thread safe to return pointers/references to array elements, in case the array is grown+relocated on another thread.
+	FOnlineAccountIdDataEOS GetAccountIdData(const FOnlineAccountIdHandle& Handle) const;
+
+	static FOnlineAccountIdRegistryEOS& Get();
 
 private:
+	friend class FAuthEOS;
+	friend class FOnlineServicesEOSModule;
+	friend FAccountIdReplicationTest;
+	FOnlineAccountIdHandle FindOrAddAccountId(const EOS_EpicAccountId EpicAccountId, const EOS_ProductUserId ProductUserId);
+
 	mutable FRWLock Lock;
 
 	TArray<FOnlineAccountIdDataEOS> AccountIdData; // Actual container for the info, indexed by the handle
 
 	TMap<EOS_EpicAccountId, FOnlineAccountIdHandle> EpicAccountIdToHandle; // Map of EOS_EpicAccountId to the associated handle.
 	TMap<EOS_ProductUserId, FOnlineAccountIdHandle> ProductUserIdToHandle; // Map of EOS_ProductUserId to the associated handle.
-
 };
 
 EOS_EpicAccountId GetEpicAccountId(const FOnlineAccountIdHandle& Handle);
 EOS_EpicAccountId GetEpicAccountIdChecked(const FOnlineAccountIdHandle& Handle);
-EOS_ProductUserId GetProductUserId(const FOnlineAccountIdHandle& Handle);
-EOS_ProductUserId GetProductUserIdChecked(const FOnlineAccountIdHandle& Handle);
-
 FOnlineAccountIdHandle FindAccountId(const EOS_EpicAccountId EpicAccountId);
-FOnlineAccountIdHandle FindAccountId(const EOS_ProductUserId EpicAccountId);
 FOnlineAccountIdHandle FindAccountIdChecked(const EOS_EpicAccountId EpicAccountId);
-FOnlineAccountIdHandle FindAccountIdChecked(const EOS_ProductUserId EpicAccountId);
-
-FOnlineAccountIdHandle CreateAccountId(const EOS_EpicAccountId EpicAccountId, const EOS_ProductUserId ProductUserId);
-
-template<typename IdType>
-inline bool ValidateOnlineId(const TOnlineIdHandle<IdType>& Handle)
-{
-	return Handle.GetOnlineServicesType() == EOnlineServices::Epic && Handle.IsValid();
-}
 
 } /* namespace UE::Online */
diff --git a/Engine/Plugins/Online/Experimental/OnlineServicesEOS/Source/Private/Online/OnlineServicesEOS.cpp b/Engine/Plugins/Online/Experimental/OnlineServicesEOS/Source/Private/Online/OnlineServicesEOS.cpp
--- a/Engine/Plugins/Online/Experimental/OnlineServicesEOS/Source/Private/Online/OnlineServicesEOS.cpp
+++ b/Engine/Plugins/Online/Experimental/OnlineServicesEOS/Source/Private/Online/OnlineServicesEOS.cpp
@@ -1,172 +1,27 @@
 // Copyright Epic Games, Inc. All Rights Reserved.
 
 #include "Online/OnlineServicesEOS.h"
-#include "Online/OnlineServicesEOSTypes.h"
-
-#include "EOSShared.h"
-#include "EOSSharedTypes.h"
-#include "IEOSSDKManager.h"
 
 #include "Online/AuthEOS.h"
 #include "Online/FriendsEOS.h"
-#include "Online/LobbiesEOS.h"
-#include "Online/OnlineIdEOS.h"
 #include "Online/PresenceEOS.h"
 #include "Online/ExternalUIEOS.h"
 
-#if WITH_ENGINE
-#include "InternetAddrEOS.h"
-#include "NetDriverEOSBase.h"
-#include "SocketSubsystemEOSUtils_OnlineServicesEOS.h"
-#endif
-
 namespace UE::Online {
 
-struct FEOSPlatformConfig
-{
-	FString ProductId;
-	FString SandboxId;
-	FString DeploymentId;
-	FString ClientId;
-	FString ClientSecret;
-	bool bUseEAS = false;
-};
-
-namespace Meta {
-
-BEGIN_ONLINE_STRUCT_META(FEOSPlatformConfig)
-	ONLINE_STRUCT_FIELD(FEOSPlatformConfig, ProductId),
-	ONLINE_STRUCT_FIELD(FEOSPlatformConfig, SandboxId),
-	ONLINE_STRUCT_FIELD(FEOSPlatformConfig, DeploymentId),
-	ONLINE_STRUCT_FIELD(FEOSPlatformConfig, ClientId),
-	ONLINE_STRUCT_FIELD(FEOSPlatformConfig, ClientSecret),
-	ONLINE_STRUCT_FIELD(FEOSPlatformConfig, bUseEAS)
-END_ONLINE_STRUCT_META()
-
-/* Meta */ }
-
 FOnlineServicesEOS::FOnlineServicesEOS(FName InInstanceName)
-	: FOnlineServicesCommon(TEXT("EOS"), InInstanceName)
+	: Super(InInstanceName)
 {
 }
 
 void FOnlineServicesEOS::RegisterComponents()
 {
-	FEOSPlatformConfig EOSPlatformConfig;
-	LoadConfig(EOSPlatformConfig);
-
-	Components.Register<FAuthEOS>(*this, EOSPlatformConfig.bUseEAS);
-	if (EOSPlatformConfig.bUseEAS)
-	{
-		Components.Register<FFriendsEOS>(*this);
-		Components.Register<FPresenceEOS>(*this);
-	}
-	Components.Register<FLobbiesEOS>(*this);
+	Components.Register<FAuthEOS>(*this);
+	Components.Register<FFriendsEOS>(*this);
+	Components.Register<FPresenceEOS>(*this);
 	Components.Register<FExternalUIEOS>(*this);
-	FOnlineServicesCommon::RegisterComponents();
-}
-
-void FOnlineServicesEOS::Initialize()
-{
-	IEOSSDKManager* SDKManager = IEOSSDKManager::Get();
-	if (!SDKManager)
-	{
-		return;
-	}
-
-	EOS_EResult InitResult = SDKManager->Initialize();
-	if (InitResult != EOS_EResult::EOS_Success)
-	{
-		return;
-	}
-
-	FEOSPlatformConfig EOSPlatformConfig;
-	LoadConfig(EOSPlatformConfig);
 
-	const FTCHARToUTF8 ProductId(*EOSPlatformConfig.ProductId);
-	const FTCHARToUTF8 SandboxId(*EOSPlatformConfig.SandboxId);
-	const FTCHARToUTF8 DeploymentId(*EOSPlatformConfig.DeploymentId);
-	const FTCHARToUTF8 ClientId(*EOSPlatformConfig.ClientId);
-	const FTCHARToUTF8 ClientSecret(*EOSPlatformConfig.ClientSecret);
-	const FTCHARToUTF8 CacheDirectory(*(SDKManager->GetCacheDirBase() / TEXT("OnlineServicesEOS")));
-
-	EOS_Platform_Options PlatformOptions = {};
-	PlatformOptions.ApiVersion = EOS_PLATFORM_OPTIONS_API_LATEST;
-	PlatformOptions.Reserved = nullptr;
-	PlatformOptions.bIsServer = EOS_FALSE;
-	PlatformOptions.OverrideCountryCode = nullptr;
-	PlatformOptions.OverrideLocaleCode = nullptr;
-	PlatformOptions.Flags = EOS_PF_WINDOWS_ENABLE_OVERLAY_D3D9 | EOS_PF_WINDOWS_ENABLE_OVERLAY_D3D10 | EOS_PF_WINDOWS_ENABLE_OVERLAY_OPENGL; // Enable overlay support for D3D9/10 and OpenGL. This sample uses D3D11 or SDL.
-
-	PlatformOptions.ProductId = ProductId.Get();
-	PlatformOptions.SandboxId = SandboxId.Get();
-	PlatformOptions.DeploymentId = DeploymentId.Get();
-	PlatformOptions.ClientCredentials.ClientId = ClientId.Get();
-	PlatformOptions.ClientCredentials.ClientSecret = ClientSecret.Get();
-	PlatformOptions.CacheDirectory = CacheDirectory.Get();
-
-	EOSPlatformHandle = SDKManager->CreatePlatform(PlatformOptions);
-	if (EOSPlatformHandle)
-	{
-#if WITH_ENGINE
-		SocketSubsystem = MakeShareable(new FSocketSubsystemEOS(EOSPlatformHandle, MakeShareable(new FSocketSubsystemEOSUtils_OnlineServicesEOS(*this))));
-		check(SocketSubsystem);
-
-		FString ErrorStr;
-		if (!SocketSubsystem->Init(ErrorStr))
-		{
-			UE_LOG(LogOnlineServices, Warning, TEXT("[FOnlineServicesEOS::Initialize] Unable to initialize Socket Subsystem. Error=[%s]"), *ErrorStr);
-		}
-#endif
-	}
-	else
-	{
-		UE_LOG(LogOnlineServices, Verbose, TEXT("[FOnlineServicesEOS::Initialize] Unable to initialize Socket Subsystem. EOS Platform Handle was invalid."));
-		return;
-	}
-
-	FOnlineServicesCommon::Initialize();
+	Super::RegisterComponents();
 }
 
-TOnlineResult<FGetResolvedConnectString> FOnlineServicesEOS::GetResolvedConnectString(FGetResolvedConnectString::Params&& Params)
-{
-	ILobbiesPtr LobbiesEOS = GetLobbiesInterface();
-	check(LobbiesEOS);
-
-	TOnlineResult<FGetJoinedLobbies> JoinedLobbies = LobbiesEOS->GetJoinedLobbies({ Params.LocalUserId });
-	if (JoinedLobbies.IsOk())
-	{
-		for (TSharedRef<const FLobby>& Lobby : JoinedLobbies.GetOkValue().Lobbies)
-		{
-			if (Lobby->LobbyId == Params.LobbyId)
-			{
-#if WITH_ENGINE
- 				//It should look like this: "EOS:0002aeeb5b2d4388a3752dd6d31222ec:GameNetDriver:97"
- 				FString NetDriverName = GetDefault<UNetDriverEOSBase>()->NetDriverName.ToString();
- 				FInternetAddrEOS TempAddr(GetProductUserIdChecked(Lobby->OwnerAccountId), NetDriverName, GetTypeHash(NetDriverName));
- 				return TOnlineResult<FGetResolvedConnectString>({ TempAddr.ToString(true) });
-#else
-				return TOnlineResult<FGetResolvedConnectString>(Errors::NotImplemented());
-#endif
-			}
-		}
-		// No matching lobby
-		return TOnlineResult<FGetResolvedConnectString>(Errors::InvalidParams());
-	}
-	else
-	{
-		return TOnlineResult<FGetResolvedConnectString>(JoinedLobbies.GetErrorValue());
-	}
-}
-
-EOS_HPlatform FOnlineServicesEOS::GetEOSPlatformHandle() const
-{
-	if (EOSPlatformHandle)
-	{
-		return *EOSPlatformHandle;
-	}
-	return nullptr;
-}
-
-
 /* UE::Online */ }
diff --git a/Engine/Plugins/Online/Experimental/OnlineServicesEOS/Source/Private/Online/SocketSubsystemEOSUtils_OnlineServicesEOS.cpp b/Engine/Plugins/Online/Experimental/OnlineServicesEOS/Source/Private/Online/SocketSubsystemEOSUtils_OnlineServicesEOS.cpp
deleted file mode 100644
diff --git a/Engine/Plugins/Online/Experimental/OnlineServicesEOS/Source/Private/Online/SocketSubsystemEOSUtils_OnlineServicesEOS.h b/Engine/Plugins/Online/Experimental/OnlineServicesEOS/Source/Private/Online/SocketSubsystemEOSUtils_OnlineServicesEOS.h
deleted file mode 100644
diff --git a/Engine/Plugins/Online/Experimental/OnlineServicesEOS/Source/Private/Online/Tests/OnlineIdEOSTests.cpp b/Engine/Plugins/Online/Experimental/OnlineServicesEOS/Source/Private/Online/Tests/OnlineIdEOSTests.cpp
--- a/Engine/Plugins/Online/Experimental/OnlineServicesEOS/Source/Private/Online/Tests/OnlineIdEOSTests.cpp
+++ b/Engine/Plugins/Online/Experimental/OnlineServicesEOS/Source/Private/Online/Tests/OnlineIdEOSTests.cpp
@@ -21,7 +21,7 @@ bool FAccountIdReplicationTest::RunTest(const FString& Parameters)
 
 		{
 			FOnlineAccountIdRegistryEOS Registry;
-			const FOnlineAccountIdHandle Handle = Registry.Create(EasId, EosId);
+			const FOnlineAccountIdHandle Handle = Registry.FindOrAddAccountId(EasId, EosId);
 			TArray<uint8> RepData = Registry.ToReplicationData(Handle);
 			const FOnlineAccountIdHandle Handle2 = Registry.FromReplicationData(RepData);
 			UTEST_EQUAL(TEXT(""), Handle, Handle2);
@@ -29,7 +29,7 @@ bool FAccountIdReplicationTest::RunTest(const FString& Parameters)
 
 		{
 			FOnlineAccountIdRegistryEOS Registry;
-			const FOnlineAccountIdHandle Handle = Registry.Create(EasId, nullptr);
+			const FOnlineAccountIdHandle Handle = Registry.FindOrAddAccountId(EasId, nullptr);
 			TArray<uint8> RepData = Registry.ToReplicationData(Handle);
 			const FOnlineAccountIdHandle Handle2 = Registry.FromReplicationData(RepData);
 			UTEST_EQUAL(TEXT(""), Handle, Handle2);
@@ -37,7 +37,7 @@ bool FAccountIdReplicationTest::RunTest(const FString& Parameters)
 
 		{
 			FOnlineAccountIdRegistryEOS Registry;
-			const FOnlineAccountIdHandle Handle = Registry.Create(nullptr, EosId);
+			const FOnlineAccountIdHandle Handle = Registry.FindOrAddAccountId(nullptr, EosId);
 			TArray<uint8> RepData = Registry.ToReplicationData(Handle);
 			const FOnlineAccountIdHandle Handle2 = Registry.FromReplicationData(RepData);
 			UTEST_EQUAL(TEXT(""), Handle, Handle2);
diff --git a/Engine/Plugins/Online/Experimental/OnlineServicesEOS/Source/Private/OnlineServicesEOSModule.cpp b/Engine/Plugins/Online/Experimental/OnlineServicesEOS/Source/Private/OnlineServicesEOSModule.cpp
--- a/Engine/Plugins/Online/Experimental/OnlineServicesEOS/Source/Private/OnlineServicesEOSModule.cpp
+++ b/Engine/Plugins/Online/Experimental/OnlineServicesEOS/Source/Private/OnlineServicesEOSModule.cpp
@@ -31,8 +31,10 @@ protected:
 
 void FOnlineServicesEOSModule::StartupModule()
 {
-	FOnlineServicesRegistry::Get().RegisterServicesFactory(EOnlineServices::Epic, MakeUnique<FOnlineServicesFactoryEOS>());
-	FOnlineIdRegistryRegistry::Get().RegisterAccountIdRegistry(EOnlineServices::Epic, &FOnlineAccountIdRegistryEOS::Get());
+	// Make this higher priority that EOSGS
+	const int Priority = 1;
+	FOnlineServicesRegistry::Get().RegisterServicesFactory(EOnlineServices::Epic, MakeUnique<FOnlineServicesFactoryEOS>(), Priority);
+	FOnlineIdRegistryRegistry::Get().RegisterAccountIdRegistry(EOnlineServices::Epic, &FOnlineAccountIdRegistryEOS::Get(), Priority);
 }
 
 /* UE::Online */ }
diff --git a/Engine/Plugins/Online/Experimental/OnlineServicesEOS/Source/Public/Online/AuthEOS.h b/Engine/Plugins/Online/Experimental/OnlineServicesEOS/Source/Public/Online/AuthEOS.h
--- a/Engine/Plugins/Online/Experimental/OnlineServicesEOS/Source/Public/Online/AuthEOS.h
+++ b/Engine/Plugins/Online/Experimental/OnlineServicesEOS/Source/Public/Online/AuthEOS.h
@@ -2,67 +2,46 @@
 
 #pragma once
 
-#include "CoreMinimal.h"
-#include "Online/AuthCommon.h"
-
-#if defined(EOS_PLATFORM_BASE_FILE_NAME)
-#include EOS_PLATFORM_BASE_FILE_NAME
-#endif
-#include "eos_auth_types.h"
-#include "eos_connect_types.h"
+#include "Online/AuthEOSGS.h"
 
 namespace UE::Online {
 
 class FOnlineServicesEOS;
 
-class ONLINESERVICESEOS_API FAuthEOS : public FAuthCommon
+class ONLINESERVICESEOS_API FAuthEOS : public FAuthEOSGS
 {
 public:
-	using Super = FAuthCommon;
+	using Super = FAuthEOSGS;
+
+	FAuthEOS(FOnlineServicesEOS& InOwningSubsystem);
+	virtual ~FAuthEOS() = default;
 
-	FAuthEOS(FOnlineServicesEOS& InOwningSubsystem, bool bInUseEAS);
+	// Begin IAuth
 	virtual void Initialize() override;
-	virtual void PreShutdown() override;
 	virtual TOnlineAsyncOpHandle<FAuthLogin> Login(FAuthLogin::Params&& Params) override;
 	virtual TOnlineAsyncOpHandle<FAuthLogout> Logout(FAuthLogout::Params&& Params) override;
-	virtual TOnlineResult<FAuthGetAccountByPlatformUserId> GetAccountByPlatformUserId(FAuthGetAccountByPlatformUserId::Params&& Params) override;
-	virtual TOnlineResult<FAuthGetAccountByAccountId> GetAccountByAccountId(FAuthGetAccountByAccountId::Params&& Params) override;
+	// End IAuth
 
-	bool IsLoggedIn(const FOnlineAccountIdHandle& AccountId) const;
+	// Begin FAuthEOSGS
+	virtual TFuture<FOnlineAccountIdHandle> ResolveAccountId(const FOnlineAccountIdHandle& LocalUserId, const EOS_ProductUserId ProductUserId) override;
+	virtual TFuture<TArray<FOnlineAccountIdHandle>> ResolveAccountIds(const FOnlineAccountIdHandle& LocalUserId, const TArray<EOS_ProductUserId>& ProductUserIds) override;
+	virtual TFunction<TFuture<FOnlineAccountIdHandle>(FOnlineAsyncOp& InAsyncOp, const EOS_ProductUserId& ProductUserId)> ResolveProductIdFn() override;
+	virtual TFunction<TFuture<TArray<FOnlineAccountIdHandle>>(FOnlineAsyncOp& InAsyncOp, const TArray<EOS_ProductUserId>& ProductUserIds)> ResolveProductIdsFn() override;
+	// End FAuthEOSGS
 
 	TFuture<FOnlineAccountIdHandle> ResolveAccountId(const FOnlineAccountIdHandle& LocalUserId, const EOS_EpicAccountId EpicAccountId);
-	TFuture<FOnlineAccountIdHandle> ResolveAccountId(const FOnlineAccountIdHandle& LocalUserId, const EOS_ProductUserId ProductUserId);
 	TFuture<TArray<FOnlineAccountIdHandle>> ResolveAccountIds(const FOnlineAccountIdHandle& LocalUserId, const TArray<EOS_EpicAccountId>& EpicAccountIds);
-	TFuture<TArray<FOnlineAccountIdHandle>> ResolveAccountIds(const FOnlineAccountIdHandle& LocalUserId, const TArray<EOS_ProductUserId>& ProductUserIds);
-
 	TFunction<TFuture<FOnlineAccountIdHandle>(FOnlineAsyncOp& InAsyncOp, const EOS_EpicAccountId& EpicAccountId)> ResolveEpicIdFn();
-	TFunction<TFuture<FOnlineAccountIdHandle>(FOnlineAsyncOp& InAsyncOp, const EOS_ProductUserId& ProductUserId)> ResolveProductIdFn();
 	TFunction<TFuture<TArray<FOnlineAccountIdHandle>>(FOnlineAsyncOp& InAsyncOp, const TArray<EOS_EpicAccountId>& EpicAccountIds)> ResolveEpicIdsFn();
-	TFunction<TFuture<TArray<FOnlineAccountIdHandle>>(FOnlineAsyncOp& InAsyncOp, const TArray<EOS_ProductUserId>& ProductUserIds)> ResolveProductIdsFn();
 
 protected:
-	TOnlineChainableAsyncOp<FAuthLogin, TSharedPtr<class FEOSConnectLoginCredentials>> LoginEAS(TOnlineAsyncOp<FAuthLogin>& InAsyncOp);
-	TSharedPtr<class FEOSConnectLoginCredentials> MakeConnectLoginCredentials(TOnlineAsyncOp<FAuthLogin>& InAsyncOp);
-
-	void OnEOSLoginStatusChanged(FOnlineAccountIdHandle LocalUserId, ELoginStatus PreviousStatus, ELoginStatus CurrentStatus);
-	TResult<FOnlineAccountIdHandle, FOnlineError> GetAccountIdByPlatformUserId(FPlatformUserId PlatformUserId) const;
-
+	TOnlineChainableAsyncOp<FAuthLogin, TSharedPtr<FEOSConnectLoginCredentials>> LoginEAS(TOnlineAsyncOp<FAuthLogin>& InAsyncOp);
 	void ProcessSuccessfulLogin(TOnlineAsyncOp<FAuthLogin>& InAsyncOp);
+	void OnEASLoginStatusChanged(FOnlineAccountIdHandle LocalUserId, ELoginStatus PreviousStatus, ELoginStatus CurrentStatus);
 
-	class FAccountInfoEOS : public FAccountInfo
-	{
-	public:
-
-	};
-
-	TMap<FOnlineAccountIdHandle, TSharedRef<FAccountInfoEOS>> AccountInfos;
-
-	EOS_HAuth AuthHandle;
-	EOS_HConnect ConnectHandle;
-	EOS_NotificationId NotifyLoginStatusChangedNotificationId = 0;
+	static FOnlineAccountIdHandle CreateAccountId(const EOS_EpicAccountId EpicAccountId, const EOS_ProductUserId ProductUserId);
 
-	/** Are we configured to use EAS? */
-	bool bUseEAS = false;
+	EOS_NotificationId NotifyEASLoginStatusChangedNotificationId = EOS_INVALID_NOTIFICATIONID;
 };
 
 /* UE::Online */ }
diff --git a/Engine/Plugins/Online/Experimental/OnlineServicesEOS/Source/Public/Online/LobbiesEOS.h b/Engine/Plugins/Online/Experimental/OnlineServicesEOS/Source/Public/Online/LobbiesEOS.h
deleted file mode 100644
diff --git a/Engine/Plugins/Online/Experimental/OnlineServicesEOS/Source/Public/Online/OnlineServicesEOS.h b/Engine/Plugins/Online/Experimental/OnlineServicesEOS/Source/Public/Online/OnlineServicesEOS.h
--- a/Engine/Plugins/Online/Experimental/OnlineServicesEOS/Source/Public/Online/OnlineServicesEOS.h
+++ b/Engine/Plugins/Online/Experimental/OnlineServicesEOS/Source/Public/Online/OnlineServicesEOS.h
@@ -2,7 +2,7 @@
 
 #pragma once
 
-#include "Online/OnlineServicesCommon.h"
+#include "Online/OnlineServicesEOSGS.h"
 
 #if defined(EOS_PLATFORM_BASE_FILE_NAME)
 #include EOS_PLATFORM_BASE_FILE_NAME
@@ -17,24 +17,15 @@ using IEOSPlatformHandlePtr = TSharedPtr<class IEOSPlatformHandle>;
 
 namespace UE::Online {
 
-class ONLINESERVICESEOS_API FOnlineServicesEOS : public FOnlineServicesCommon
+class ONLINESERVICESEOS_API FOnlineServicesEOS : public FOnlineServicesEOSGS
 {
 public:
-	using Super = FOnlineServicesCommon;
+	using Super = FOnlineServicesEOSGS;
 
 	FOnlineServicesEOS(FName InInstanceName);
-	virtual void RegisterComponents() override;
-	virtual void Initialize() override;
-	virtual TOnlineResult<FGetResolvedConnectString> GetResolvedConnectString(FGetResolvedConnectString::Params&& Params) override;
-	virtual EOnlineServices GetServicesProvider() const override { return EOnlineServices::Epic; }
-
-	EOS_HPlatform GetEOSPlatformHandle() const;
-protected:
-	IEOSPlatformHandlePtr EOSPlatformHandle;
+	virtual ~FOnlineServicesEOS() = default;
 
-#if WITH_ENGINE
-	TSharedPtr<FSocketSubsystemEOS, ESPMode::ThreadSafe> SocketSubsystem;
-#endif
+	virtual void RegisterComponents() override;
 };
 
 /* UE::Online */ }
diff --git a/Engine/Plugins/Online/Experimental/OnlineServicesEOS/Source/Public/Online/OnlineServicesEOSTypes.h b/Engine/Plugins/Online/Experimental/OnlineServicesEOS/Source/Public/Online/OnlineServicesEOSTypes.h
--- a/Engine/Plugins/Online/Experimental/OnlineServicesEOS/Source/Public/Online/OnlineServicesEOSTypes.h
+++ b/Engine/Plugins/Online/Experimental/OnlineServicesEOS/Source/Public/Online/OnlineServicesEOSTypes.h
@@ -2,218 +2,8 @@
 
 #pragma once
 
-#include "Online/OnlineAsyncOp.h"
-#include "Online/OnlineError.h"
-#include "Async/Future.h"
-
-#include "EOSShared.h"
+#include "Online/OnlineServicesEOSGSTypes.h"
 
 namespace UE::Online {
 
-inline FOnlineError FromEOSError(EOS_EResult ResultCode)
-{
-	// Todo: make this for real.
-	return FOnlineError(Errors::Unknown());
-}
-
-/** Class to handle all callbacks generically using a future to forward callback results */
-template<typename CallbackType>
-class TEOSCallback
-{
-	using CallbackFuncType = void (EOS_CALL*)(const CallbackType*);
-public:
-	TEOSCallback() = default;
-	TEOSCallback(TPromise<const CallbackType*>&& InPromise) : Promise(MoveTemp(InPromise)) {}
-	virtual ~TEOSCallback() = default;
-
-	operator CallbackFuncType()
-	{
-		return &CallbackImpl;
-	}
-
-	TFuture<const CallbackType*> GetFuture() { return Promise.GetFuture(); }
-
-private:
-	TPromise<const CallbackType*> Promise;
-
-	static void EOS_CALL CallbackImpl(const CallbackType* Data)
-	{
-		if (EOS_EResult_IsOperationComplete(Data->ResultCode) == EOS_FALSE)
-		{
-			// Ignore
-			return;
-		}
-		check(IsInGameThread());
-
-		TEOSCallback* CallbackThis = reinterpret_cast<TEOSCallback*>(Data->ClientData);
-		check(CallbackThis);
-
-		CallbackThis->Promise.EmplaceValue(Data);
-
-		delete CallbackThis;
-	}
-};
-
-template<typename TEOSResult, typename TEOSHandle, typename TEOSParameters, typename TEOSFn> 
-TFuture<const TEOSResult*> EOS_Async(TEOSFn EOSFn, TEOSHandle EOSHandle, TEOSParameters Parameters)
-{
-	TEOSCallback<TEOSResult>* Callback = new TEOSCallback<TEOSResult>();
-	EOSFn(EOSHandle, &Parameters, Callback, *Callback);
-	return Callback->GetFuture();
-}
-
-template<typename TEOSResult, typename TEOSHandle, typename TEOSParameters, typename TEOSFn>
-void EOS_Async(TEOSFn EOSFn, TEOSHandle EOSHandle, TEOSParameters Parameters, TPromise<const TEOSResult*>&& Promise)
-{
-	TEOSCallback<TEOSResult>* Callback = new TEOSCallback<TEOSResult>(MoveTemp(Promise));
-	EOSFn(EOSHandle, &Parameters, Callback, *Callback);
-}
-
-class EOSEventRegistration
-{
-public:
-	virtual ~EOSEventRegistration() = default;
-};
-typedef TUniquePtr<EOSEventRegistration> EOSEventRegistrationPtr;
-
-/** 
-* EOS event registration utility for binding an EOS notifier registration to a RAII object which handles
-* unregistering when it exits scope. Intended to be used from a TOnlineComponent class.
-* 
-* Example:
-*	EOSEventRegistrationPtr OnLobbyUpdatedEOSEventRegistration = EOS_RegisterComponentEventHandler(
-*		this,
-*		LobbyHandle,
-*		EOS_LOBBY_ADDNOTIFYLOBBYUPDATERECEIVED_API_LATEST,
-*		&EOS_Lobby_AddNotifyLobbyUpdateReceived,
-*		&EOS_Lobby_RemoveNotifyLobbyUpdateReceived,
-*		&FLobbiesEOS::HandleLobbyUpdated);
-**/
-template <
-	typename ComponentHandlerClass,
-	typename EOSHandle,
-	typename EOSNotfyRegisterFunction,
-	typename EOSNotfyUnregisterFunction,
-	typename ComponentHandlerFunction>
-EOSEventRegistrationPtr EOS_RegisterComponentEventHandler(
-	ComponentHandlerClass* HandlerClass,
-	EOSHandle ClientHandle,
-	int32_t ApiVersion,
-	EOSNotfyRegisterFunction NotfyRegisterFunction,
-	EOSNotfyUnregisterFunction NotfyUnregisterFunction,
-	ComponentHandlerFunction HandlerFunction);
-
-namespace detail {
-
-template<typename Function> struct TEOSCallbackTraitsBase;
-
-template<typename Component, typename EventData>
-struct TEOSCallbackTraitsBase<void (Component::*)(const EventData*)>
-{
-	typedef Component ComponentType;
-	typedef EventData EventDataType;
-};
-
-template<typename Function>
-struct TEOSCallbackTraits : public TEOSCallbackTraitsBase<typename std::remove_reference<Function>::type>
-{
-};
-
-template<typename Function> struct TEOSNotifyRegisterTraits;
-
-template<typename EOSHandle, typename Options, typename ClientData, typename NotificationFn, typename NotificationId>
-struct TEOSNotifyRegisterTraits<NotificationId (*)(EOSHandle, const Options*, ClientData*, NotificationFn)>
-{
-	typedef Options OptionsType;
-	typedef NotificationId NotificationIdType;
-};
-
-template <
-	typename ComponentHandlerClass,
-	typename EOSHandle,
-	typename EOSNotfyRegisterFunction,
-	typename EOSNotfyUnregisterFunction,
-	typename ComponentHandlerFunction>
-class EOSEventRegistrationImpl : public EOSEventRegistration
-{
-public:
-	typedef EOSEventRegistrationImpl<
-		ComponentHandlerClass,
-		EOSHandle,
-		EOSNotfyRegisterFunction,
-		EOSNotfyUnregisterFunction,
-		ComponentHandlerFunction> ThisClass;
-
-	EOSEventRegistrationImpl(
-		ComponentHandlerClass* HandlerClass,
-		EOSHandle ClientHandle,
-		int32_t ApiVersion,
-		EOSNotfyRegisterFunction NotfyRegisterFunction,
-		EOSNotfyUnregisterFunction NotfyUnregisterFunction,
-		ComponentHandlerFunction HandlerFunction)
-		: HandlerClass(HandlerClass)
-		, ClientHandle(ClientHandle)
-		, NotfyUnregisterFunction(NotfyUnregisterFunction)
-		, HandlerFunction(HandlerFunction)
-	{
-		typename TEOSNotifyRegisterTraits<EOSNotfyRegisterFunction>::OptionsType Options = { };
-		Options.ApiVersion = ApiVersion;
-		NotificationId = NotfyRegisterFunction(ClientHandle, &Options, this,
-		[](const typename TEOSCallbackTraits<ComponentHandlerFunction>::EventDataType* Data)
-		{
-			ThisClass* This = reinterpret_cast<ThisClass*>(Data->ClientData);
-			(This->HandlerClass->*This->HandlerFunction)(Data);
-		});
-	}
-
-	virtual ~EOSEventRegistrationImpl()
-	{
-		NotfyUnregisterFunction(ClientHandle, NotificationId);
-	}
-
-private:
-	EOSEventRegistrationImpl() = delete;
-	EOSEventRegistrationImpl(const EOSEventRegistrationImpl&) = delete;
-	EOSEventRegistrationImpl& operator=(const EOSEventRegistrationImpl&) = delete;
-
-	typedef typename TEOSNotifyRegisterTraits<EOSNotfyRegisterFunction>::NotificationIdType NotificationIdType;
-
-	NotificationIdType NotificationId;
-	ComponentHandlerClass* HandlerClass;
-	EOSHandle ClientHandle;
-	EOSNotfyUnregisterFunction NotfyUnregisterFunction;
-	ComponentHandlerFunction HandlerFunction;
-};
-
-/* detail */ }
-
-template <
-	typename ComponentHandlerClass,
-	typename EOSHandle,
-	typename EOSNotfyRegisterFunction,
-	typename EOSNotfyUnregisterFunction,
-	typename ComponentHandlerFunction>
-EOSEventRegistrationPtr EOS_RegisterComponentEventHandler(
-	ComponentHandlerClass* HandlerClass,
-	EOSHandle ClientHandle,
-	int32_t ApiVersion,
-	EOSNotfyRegisterFunction NotfyRegisterFunction,
-	EOSNotfyUnregisterFunction NotfyUnregisterFunction,
-	ComponentHandlerFunction HandlerFunction)
-{
-	return MakeUnique<
-		detail::EOSEventRegistrationImpl<
-			ComponentHandlerClass,
-			EOSHandle,
-			EOSNotfyRegisterFunction,
-			EOSNotfyUnregisterFunction,
-			ComponentHandlerFunction>>(
-			HandlerClass,
-			ClientHandle,
-			ApiVersion,
-			NotfyRegisterFunction,
-			NotfyUnregisterFunction,
-			HandlerFunction);
-}
-
 /* UE::Online */ }
diff --git a/Engine/Plugins/Online/Experimental/OnlineServicesEOSGS/OnlineServicesEOSGS.uplugin b/Engine/Plugins/Online/Experimental/OnlineServicesEOSGS/OnlineServicesEOSGS.uplugin
new file mode 100644
--- /dev/null
+++ b/Engine/Plugins/Online/Experimental/OnlineServicesEOSGS/OnlineServicesEOSGS.uplugin
@@ -0,0 +1,45 @@
+{
+	"FileVersion": 3,
+
+	"FriendlyName": "Online Services EOS (Game Services)",
+	"Version": 1,
+	"VersionName": "1.0",
+	"Description": "Online Services implementation for EOS Game services only.",
+	"Category": "Online Platform",
+	"CreatedBy": "Epic Games, Inc.",
+	"CreatedByURL": "http://epicgames.com",
+	"EnabledByDefault": false,
+
+	"Modules": [
+		{
+			"Name": "OnlineServicesEOSGS",
+			"Type": "Runtime",
+			"LoadingPhase": "PostConfigInit",
+			"PlatformDenyList": [
+				"LinuxArm64",
+				"Hololens",
+				"TVOS"
+			]
+		}
+	],
+	"LocalizationTargets": [
+		{
+			"Name": "OnlineSubsystemEOSGS",
+			"LoadingPolicy": "Always"
+		}
+	],
+	"Plugins": [
+		{
+			"Name": "OnlineServices",
+			"Enabled": true
+		},
+		{
+			"Name": "EOSShared",
+			"Enabled": true
+		},
+		{
+			"Name": "SocketSubsystemEOS",
+			"Enabled": true
+		}
+	]
+}
diff --git a/Engine/Plugins/Online/Experimental/OnlineServicesEOSGS/Source/OnlineServicesEOSGS.Build.cs b/Engine/Plugins/Online/Experimental/OnlineServicesEOSGS/Source/OnlineServicesEOSGS.Build.cs
new file mode 100644
--- /dev/null
+++ b/Engine/Plugins/Online/Experimental/OnlineServicesEOSGS/Source/OnlineServicesEOSGS.Build.cs
@@ -0,0 +1,38 @@
+// Copyright Epic Games, Inc. All Rights Reserved.
+
+using UnrealBuildTool;
+using System.IO;
+
+public class OnlineServicesEOSGS : ModuleRules
+{
+	public OnlineServicesEOSGS(ReadOnlyTargetRules Target) : base(Target)
+	{
+		PublicDependencyModuleNames.AddRange(
+			new string[] {
+				"Core",
+				"OnlineServicesInterface",
+				"OnlineServicesCommon",
+			}
+		);
+
+		PrivateDependencyModuleNames.AddRange(
+			new string[] {
+				"CoreOnline",
+				"CoreUObject",
+				"EOSSDK",
+				"EOSShared"
+			}
+		);
+
+		if (Target.bCompileAgainstEngine)
+		{
+			PrivateDependencyModuleNames.AddRange(
+				new string[]
+				{
+					"SocketSubsystemEOS",
+					"Sockets"
+				}
+			);
+		}
+	}
+}
diff --git a/Engine/Plugins/Online/Experimental/OnlineServicesEOSGS/Source/Private/Online/AuthEOSGS.cpp b/Engine/Plugins/Online/Experimental/OnlineServicesEOSGS/Source/Private/Online/AuthEOSGS.cpp
new file mode 100644
--- /dev/null
+++ b/Engine/Plugins/Online/Experimental/OnlineServicesEOSGS/Source/Private/Online/AuthEOSGS.cpp
@@ -0,0 +1,608 @@
+// Copyright Epic Games, Inc. All Rights Reserved.
+
+#include "Online/AuthEOSGS.h"
+
+#include "Algo/Transform.h"
+#include "Misc/CommandLine.h"
+#include "Online/AuthErrors.h"
+#include "Online/OnlineAsyncOp.h"
+#include "Online/OnlineErrorDefinitions.h"
+#include "Online/OnlineIdEOSGS.h"
+#include "Online/OnlineServicesEOSGS.h"
+#include "Online/OnlineServicesEOSGSTypes.h"
+
+#include "eos_auth.h"
+#include "eos_common.h"
+#include "eos_connect.h"
+#include "eos_init.h"
+#include "eos_logging.h"
+#include "eos_sdk.h"
+#include "eos_types.h"
+#include "eos_userinfo.h"
+
+namespace UE::Online {
+
+struct FAuthEOSConfig
+{
+	FString DefaultExternalCredentialTypeStr;
+};
+
+namespace Meta {
+
+BEGIN_ONLINE_STRUCT_META(FAuthEOSConfig)
+	ONLINE_STRUCT_FIELD(FAuthEOSConfig, DefaultExternalCredentialTypeStr)
+END_ONLINE_STRUCT_META()
+
+/* Meta*/ }
+
+ELoginStatus ToELoginStatus(EOS_ELoginStatus InStatus)
+{
+	switch (InStatus)
+	{
+	case EOS_ELoginStatus::EOS_LS_NotLoggedIn:
+	{
+		return ELoginStatus::NotLoggedIn;
+	}
+	case EOS_ELoginStatus::EOS_LS_UsingLocalProfile:
+	{
+		return ELoginStatus::UsingLocalProfile;
+	}
+	case EOS_ELoginStatus::EOS_LS_LoggedIn:
+	{
+		return ELoginStatus::LoggedIn;
+	}
+	}
+	return ELoginStatus::NotLoggedIn;
+}
+
+FEOSAuthCredentials::FEOSAuthCredentials() :
+	EOS_Auth_Credentials()
+{
+	ApiVersion = EOS_AUTH_CREDENTIALS_API_LATEST;
+	Id = IdAnsi;
+	Token = TokenAnsi;
+}
+
+
+FEOSAuthCredentials::FEOSAuthCredentials(const FEOSAuthCredentials& Other)
+{
+	ApiVersion = Other.ApiVersion;
+	Id = IdAnsi;
+	Token = TokenAnsi;
+	Type = Other.Type;
+	SystemAuthCredentialsOptions = Other.SystemAuthCredentialsOptions;
+	ExternalType = Other.ExternalType;
+
+	FCStringAnsi::Strncpy(IdAnsi, Other.IdAnsi, EOS_STRING_BUFFER_LENGTH);
+	FCStringAnsi::Strncpy(TokenAnsi, Other.TokenAnsi, EOS_MAX_TOKEN_SIZE);
+}
+
+FEOSAuthCredentials::FEOSAuthCredentials(EOS_EExternalCredentialType InExternalType, const TArray<uint8>& InToken) :
+	EOS_Auth_Credentials()
+{
+	ApiVersion = EOS_AUTH_CREDENTIALS_API_LATEST;
+	Id = IdAnsi;
+	Token = TokenAnsi;
+
+	Type = EOS_ELoginCredentialType::EOS_LCT_ExternalAuth;
+	ExternalType = InExternalType;
+
+	uint32_t InOutBufferLength = EOS_MAX_TOKEN_SIZE;
+	EOS_ByteArray_ToString(InToken.GetData(), InToken.Num(), TokenAnsi, &InOutBufferLength);
+}
+
+FEOSAuthCredentials& FEOSAuthCredentials::operator=(FEOSAuthCredentials& Other)
+{
+	ApiVersion = Other.ApiVersion;
+	Type = Other.Type;
+	SystemAuthCredentialsOptions = Other.SystemAuthCredentialsOptions;
+	ExternalType = Other.ExternalType;
+
+	FCStringAnsi::Strncpy(IdAnsi, Other.IdAnsi, EOS_STRING_BUFFER_LENGTH);
+	FCStringAnsi::Strncpy(TokenAnsi, Other.TokenAnsi, EOS_MAX_TOKEN_SIZE);
+
+	return *this;
+}
+
+void FEOSAuthCredentials::SetToken(const FCredentialsToken& InToken)
+{
+	if (InToken.IsType<TArray<uint8>>())
+	{
+		const TArray<uint8>& TokenData = InToken.Get<TArray<uint8>>();
+		uint32_t InOutBufferLength = EOS_MAX_TOKEN_SIZE;
+		EOS_ByteArray_ToString(TokenData.GetData(), TokenData.Num(), TokenAnsi, &InOutBufferLength);
+	}
+	else if (InToken.IsType<FString>())
+	{
+		const FString& TokenString = InToken.Get<FString>();
+		FCStringAnsi::Strncpy(TokenAnsi, TCHAR_TO_UTF8(*TokenString), EOS_MAX_TOKEN_SIZE);
+	}
+	else
+	{
+		UE_LOG(LogTemp, Warning, TEXT("SetToken cannot be set with an invalid credentials token parameter. Please ensure there is valid data in the credentials token."));
+	}
+}
+
+FEOSConnectLoginCredentials::FEOSConnectLoginCredentials()
+{
+	ApiVersion = EOS_CONNECT_CREDENTIALS_API_LATEST;
+}
+
+class FEOSConnectLoginCredentialsEOS : public FEOSConnectLoginCredentials
+{
+public:
+	FEOSConnectLoginCredentialsEOS()
+	{
+		Token = TokenAnsi;
+		TokenAnsi[0] = '\0';
+	}
+	void SetToken(const FCredentialsToken& InToken)
+	{
+		if (InToken.IsType<TArray<uint8>>())
+		{
+			const TArray<uint8>& TokenData = InToken.Get<TArray<uint8>>();
+			uint32_t InOutBufferLength = EOS_MAX_TOKEN_SIZE;
+			EOS_ByteArray_ToString(TokenData.GetData(), TokenData.Num(), TokenAnsi, &InOutBufferLength);
+		}
+		else if (InToken.IsType<FString>())
+		{
+			const FString& TokenString = InToken.Get<FString>();
+			FCStringAnsi::Strncpy(TokenAnsi, TCHAR_TO_UTF8(*TokenString), EOS_MAX_TOKEN_SIZE);
+		}
+		else
+		{
+			UE_LOG(LogTemp, Warning, TEXT("SetToken cannot be set with an invalid credentials token parameter. Please ensure there is valid data in the credentials token."));
+		}
+	}
+
+	char TokenAnsi[EOS_MAX_TOKEN_SIZE];
+};
+
+FEOSConnectLoginCredentialsEAS::FEOSConnectLoginCredentialsEAS(EOS_Auth_Token* InEASToken) : EASToken(InEASToken)
+{
+	check(EASToken);
+	Type = EOS_EExternalCredentialType::EOS_ECT_EPIC;
+	Token = EASToken->AccessToken;
+}
+
+FEOSConnectLoginCredentialsEAS::~FEOSConnectLoginCredentialsEAS()
+{
+	if (EASToken)
+	{
+		EOS_Auth_Token_Release(EASToken);
+	}
+}
+
+FAuthEOSGS::FAuthEOSGS(FOnlineServicesEOSGS& InServices)
+	: Super(InServices)
+{
+}
+
+void FAuthEOSGS::Initialize()
+{
+	Super::Initialize();
+
+	AuthHandle = EOS_Platform_GetAuthInterface(static_cast<FOnlineServicesEOSGS&>(GetServices()).GetEOSPlatformHandle());
+	check(AuthHandle != nullptr);
+
+	ConnectHandle = EOS_Platform_GetConnectInterface(static_cast<FOnlineServicesEOSGS&>(GetServices()).GetEOSPlatformHandle());
+	check(ConnectHandle != nullptr);
+}
+
+void SplitCredentialsType(const FString& InStr, FString& OutCredentialTypeStr, FString& OutExternalCredentialTypeStr)
+{
+	InStr.Split(FString(TEXT(":")), &OutCredentialTypeStr, &OutExternalCredentialTypeStr);
+	if (OutCredentialTypeStr.IsEmpty())
+	{
+		OutCredentialTypeStr = InStr;
+	}
+}
+
+TSharedPtr<FEOSConnectLoginCredentials> MakeConnectLoginCredentials(TOnlineAsyncOp<FAuthLogin>& InAsyncOp)
+{
+	const FAuthLogin::Params& LoginParams = InAsyncOp.GetParams();
+
+	TSharedRef<FEOSConnectLoginCredentialsEOS> Credentials = MakeShared<FEOSConnectLoginCredentialsEOS>();
+	LexFromString(Credentials->Type, *LoginParams.CredentialsType);
+	Credentials->SetToken(LoginParams.CredentialsToken);
+	return Credentials;
+}
+
+TOnlineAsyncOpHandle<FAuthLogin> FAuthEOSGS::Login(FAuthLogin::Params&& Params)
+{
+#if (WITH_EDITOR || !UE_BUILD_SHIPPING)
+	// Handle developer login
+	if (Params.CredentialsId.IsEmpty() && Params.CredentialsType.IsEmpty() && Params.CredentialsToken.IsType<FString>() && Params.CredentialsToken.Get<FString>().IsEmpty())
+	{
+		FString	CommandLineAuthId;
+		FString	CommandLineAuthToken;
+		FString	CommandLineAuthType;
+		FParse::Value(FCommandLine::Get(), TEXT("AUTH_LOGIN="), CommandLineAuthId);
+		FParse::Value(FCommandLine::Get(), TEXT("AUTH_PASSWORD="), CommandLineAuthToken);
+		FParse::Value(FCommandLine::Get(), TEXT("AUTH_TYPE="), CommandLineAuthType);
+		if (!CommandLineAuthId.IsEmpty() && !CommandLineAuthToken.IsEmpty() && !CommandLineAuthType.IsEmpty())
+		{
+			Params.CredentialsId = MoveTemp(CommandLineAuthId);
+			Params.CredentialsToken.Emplace<FString>(MoveTemp(CommandLineAuthToken));
+			Params.CredentialsType = MoveTemp(CommandLineAuthType);
+		}
+	}
+#endif
+
+	TOnlineAsyncOpRef<FAuthLogin> Op = GetOp<FAuthLogin>(MoveTemp(Params));
+	// Are we already logged in?
+	if (GetAccountIdByPlatformUserId(Op->GetParams().PlatformUserId).IsOk())
+	{
+		Op->SetError(Errors::Auth::AlreadyLoggedIn());
+		return Op->GetHandle();
+	}
+
+	TOnlineChainableAsyncOp<FAuthLogin, TSharedPtr<FEOSConnectLoginCredentials>> ConnectLoginOp = [this, &Op]()
+	{
+		FString LoginCredentialTypeStr, Unused;
+		SplitCredentialsType(Op->GetParams().CredentialsType, LoginCredentialTypeStr, Unused);
+
+		// If the Credentials type is an EAS Login credential type, then login to EOS_Auth with that and get an EAS token to login to EOS_Connect with
+		EOS_ELoginCredentialType EasLoginType;
+		if (LexFromString(EasLoginType, *LoginCredentialTypeStr))
+		{
+			return LoginEAS(*Op);
+		}
+
+		// The credentials type is presumably an EOS_EExternalCredentialType, so just login to connect with that
+		return Op->Then([this](TOnlineAsyncOp<FAuthLogin>& InAsyncOp)
+		{
+			return MakeConnectLoginCredentials(InAsyncOp);
+		});
+	}();
+	// Check if the above steps completed (failed) the operation
+	if (ConnectLoginOp.GetOwningOperation().IsComplete())
+	{
+		return ConnectLoginOp.GetOwningOperation().GetHandle();
+	}
+	ConnectLoginOp.Then([this](TOnlineAsyncOp<FAuthLogin>& InAsyncOp, TSharedPtr<FEOSConnectLoginCredentials>&& InConnectLoginCredentials)
+	{
+		EOS_Connect_LoginOptions ConnectLoginOptions = { };
+		ConnectLoginOptions.ApiVersion = EOS_CONNECT_LOGIN_API_LATEST;
+		ConnectLoginOptions.Credentials = InConnectLoginCredentials.Get();
+
+		return EOS_Async<EOS_Connect_LoginCallbackInfo>(EOS_Connect_Login, ConnectHandle, ConnectLoginOptions);
+	})
+	.Then([this](TOnlineAsyncOp<FAuthLogin>& InAsyncOp, const EOS_Connect_LoginCallbackInfo* Data)
+	{
+		UE_LOG(LogTemp, Verbose, TEXT("[FAuthEOSGS::Login] EOS_Connect_Login Result: [%s]"), *LexToString(Data->ResultCode));
+
+		if (Data->ResultCode == EOS_EResult::EOS_Success)
+		{
+			ProcessSuccessfulLogin(InAsyncOp, Data->LocalUserId);
+		}
+		else if (Data->ResultCode == EOS_EResult::EOS_InvalidUser && Data->ContinuanceToken != nullptr)
+		{
+			EOS_Connect_CreateUserOptions ConnectCreateUserOptions = { };
+			ConnectCreateUserOptions.ApiVersion = EOS_CONNECT_CREATEUSER_API_LATEST;
+			ConnectCreateUserOptions.ContinuanceToken = Data->ContinuanceToken;
+
+			return EOS_Async<EOS_Connect_CreateUserCallbackInfo>(EOS_Connect_CreateUser, ConnectHandle, ConnectCreateUserOptions);
+		}
+		else
+		{
+			// TODO: EAS Logout
+			InAsyncOp.SetError(Errors::Unknown()); // TODO
+		}
+
+		return MakeFulfilledPromise<const EOS_Connect_CreateUserCallbackInfo*>().GetFuture();
+	})
+	.Then([this](TOnlineAsyncOp<FAuthLogin>& InAsyncOp, const EOS_Connect_CreateUserCallbackInfo* Data)
+	{
+		UE_LOG(LogTemp, Verbose, TEXT("[FAuthEOSGS::Login] EOS_Connect_CreateUser Result: [%s]"), *LexToString(Data->ResultCode));
+
+		if (Data->ResultCode == EOS_EResult::EOS_Success)
+		{
+			ProcessSuccessfulLogin(InAsyncOp, Data->LocalUserId);
+		}
+		else
+		{
+			// TODO: EAS Logout
+
+			InAsyncOp.SetError(Errors::Unknown()); // TODO
+		}
+	})
+	.Enqueue(GetSerialQueue());
+
+	return Op->GetHandle();
+}
+
+TOnlineChainableAsyncOp<FAuthLogin, TSharedPtr<FEOSConnectLoginCredentials>> FAuthEOSGS::LoginEAS(TOnlineAsyncOp<FAuthLogin>& InAsyncOp)
+{
+	const FAuthLogin::Params& Params = InAsyncOp.GetParams();
+	auto FailureLambda = [](TOnlineAsyncOp<FAuthLogin>& Op) -> TSharedPtr<FEOSConnectLoginCredentials>
+	{
+		return nullptr;
+	};
+
+	EOS_Auth_LoginOptions LoginOptions = { };
+	LoginOptions.ApiVersion = EOS_AUTH_LOGIN_API_LATEST;
+
+	FString CredentialTypeStr;
+	FString ExternalCredentialTypeStr;
+	SplitCredentialsType(Params.CredentialsType, CredentialTypeStr, ExternalCredentialTypeStr);
+
+	FEOSAuthCredentials Credentials;
+	if (LexFromString(Credentials.Type, *CredentialTypeStr))
+	{
+		switch (Credentials.Type)
+		{
+		case EOS_ELoginCredentialType::EOS_LCT_ExternalAuth:
+			// If an external credential type wasn't specified, we'll grab the platform default from the configuration values
+			if (ExternalCredentialTypeStr.IsEmpty())
+			{
+				FAuthEOSConfig AuthEOSConfig;
+				LoadConfig(AuthEOSConfig);
+				ExternalCredentialTypeStr = AuthEOSConfig.DefaultExternalCredentialTypeStr;
+			}
+
+			EOS_EExternalCredentialType ExternalType;
+			if (LexFromString(ExternalType, *ExternalCredentialTypeStr))
+			{
+				Credentials.ExternalType = ExternalType;
+			}
+
+			Credentials.SetToken(Params.CredentialsToken);
+			break;
+		case EOS_ELoginCredentialType::EOS_LCT_ExchangeCode:
+			// This is how the Epic launcher will pass credentials to you
+			Credentials.IdAnsi[0] = '\0';
+			FCStringAnsi::Strncpy(Credentials.TokenAnsi, TCHAR_TO_UTF8(*Params.CredentialsToken.Get<FString>()), EOS_MAX_TOKEN_SIZE);
+			break;
+		case EOS_ELoginCredentialType::EOS_LCT_Password:
+			FCStringAnsi::Strncpy(Credentials.IdAnsi, TCHAR_TO_UTF8(*Params.CredentialsId), EOS_STRING_BUFFER_LENGTH);
+			FCStringAnsi::Strncpy(Credentials.TokenAnsi, TCHAR_TO_UTF8(*Params.CredentialsToken.Get<FString>()), EOS_MAX_TOKEN_SIZE);
+			break;
+		case EOS_ELoginCredentialType::EOS_LCT_Developer:
+			// This is auth via the EOS auth tool
+			FCStringAnsi::Strncpy(Credentials.IdAnsi, TCHAR_TO_UTF8(*Params.CredentialsId), EOS_STRING_BUFFER_LENGTH);
+			FCStringAnsi::Strncpy(Credentials.TokenAnsi, TCHAR_TO_UTF8(*Params.CredentialsToken.Get<FString>()), EOS_MAX_TOKEN_SIZE);
+			break;
+		case EOS_ELoginCredentialType::EOS_LCT_AccountPortal:
+			// This is auth via the EOS Account Portal
+			Credentials.IdAnsi[0] = '\0';
+			Credentials.TokenAnsi[0] = '\0';
+			break;
+		case EOS_ELoginCredentialType::EOS_LCT_PersistentAuth:
+			// This is auth via stored credentials in EOS
+			Credentials.Id = nullptr;
+			Credentials.Token = nullptr;
+			break;
+		default:
+			UE_LOG(LogTemp, Warning, TEXT("Unsupported CredentialsType=[%s]"), *Params.CredentialsType);
+			InAsyncOp.SetError(Errors::Unknown()); // TODO
+			return InAsyncOp.Then(FailureLambda);
+		}
+	}
+	else
+	{
+		UE_LOG(LogTemp, Warning, TEXT("Invalid CredentialsType=[%s]"), *Params.CredentialsType);
+		InAsyncOp.SetError(Errors::Unknown()); // TODO
+		return InAsyncOp.Then(FailureLambda);
+	}
+
+	return InAsyncOp.Then([this, LoginOptions, Credentials](TOnlineAsyncOp<FAuthLogin>& InAsyncOp) mutable
+	{
+		LoginOptions.Credentials = &Credentials;
+		return EOS_Async<EOS_Auth_LoginCallbackInfo>(EOS_Auth_Login, AuthHandle, LoginOptions);
+	})
+	.Then([this](TOnlineAsyncOp<FAuthLogin>& InAsyncOp, const EOS_Auth_LoginCallbackInfo* Data)
+	{
+		UE_LOG(LogTemp, Verbose, TEXT("[FAuthEOS::Login] EOS_Auth_Login Result: [%s]"), *LexToString(Data->ResultCode));
+
+		if (Data->ResultCode == EOS_EResult::EOS_Success)
+		{
+			// We cache the Epic Account Id to use it in later stages of the login process
+			InAsyncOp.Data.Set(TEXT("EpicAccountId"), Data->LocalUserId);
+
+			// A success means the account is already linked, we'll process the NULL result in the next step
+		}
+		else if (Data->ResultCode == EOS_EResult::EOS_InvalidUser && Data->ContinuanceToken != nullptr)
+		{
+			EOS_Auth_LinkAccountOptions LinkAccountOptions = {};
+			LinkAccountOptions.ApiVersion = EOS_AUTH_LINKACCOUNT_API_LATEST;
+			LinkAccountOptions.ContinuanceToken = Data->ContinuanceToken;
+
+			return EOS_Async<EOS_Auth_LinkAccountCallbackInfo>(EOS_Auth_LinkAccount, AuthHandle, LinkAccountOptions);
+		}
+		else
+		{
+			FOnlineError Error = Errors::Unknown();
+			if (Data->ResultCode == EOS_EResult::EOS_InvalidAuth)
+			{
+				Error = Errors::InvalidCreds();
+			}
+
+			InAsyncOp.SetError(MoveTemp(Error));
+		}
+
+		return MakeFulfilledPromise<const EOS_Auth_LinkAccountCallbackInfo*>().GetFuture();
+	})
+	.Then([this](TOnlineAsyncOp<FAuthLogin>& InAsyncOp, const EOS_Auth_LinkAccountCallbackInfo* Data) -> TFuture<TSharedPtr<FEOSConnectLoginCredentials>>
+	{
+		UE_LOG(LogTemp, Verbose, TEXT("[FAuthEOS::Login] EOS_Auth_LinkAccount Result: [%s]"), (Data == nullptr) ? TEXT("Null") : *LexToString(Data->ResultCode));
+
+		// If Data is NULL, it means the account was already linked
+		if (Data == nullptr || Data->ResultCode == EOS_EResult::EOS_Success)
+		{
+			EOS_EpicAccountId EpicAccountId = nullptr;
+			if (Data == nullptr)
+			{
+				EpicAccountId = *InAsyncOp.Data.Get<EOS_EpicAccountId>(TEXT("EpicAccountId"));
+			}
+			else
+			{
+				// We cache the Epic Account Id to use it in later stages of the login process
+				InAsyncOp.Data.Set(TEXT("EpicAccountId"), Data->LocalUserId);
+				EpicAccountId = Data->LocalUserId;
+			}
+
+			// On success, attempt Connect Login
+			EOS_Auth_Token* AuthToken = nullptr;
+			EOS_Auth_CopyUserAuthTokenOptions CopyOptions = { };
+			CopyOptions.ApiVersion = EOS_AUTH_COPYUSERAUTHTOKEN_API_LATEST;
+
+			EOS_EResult CopyResult = EOS_Auth_CopyUserAuthToken(AuthHandle, &CopyOptions, EpicAccountId, &AuthToken);
+
+			UE_LOG(LogTemp, Verbose, TEXT("[FAuthEOS::Login] EOS_Auth_CopyUserAuthToken Result: [%s]"), *LexToString(CopyResult));
+
+			if (CopyResult == EOS_EResult::EOS_Success)
+			{
+				TSharedPtr<FEOSConnectLoginCredentials> ConnectLoginCredentials = MakeShared<FEOSConnectLoginCredentialsEAS>(AuthToken);
+				return MakeFulfilledPromise<TSharedPtr<FEOSConnectLoginCredentials>>(MoveTemp(ConnectLoginCredentials)).GetFuture();
+			}
+			else
+			{
+				// TODO: EAS Logout
+				InAsyncOp.SetError(Errors::Unknown()); // TODO
+			}
+		}
+		else
+		{
+			InAsyncOp.SetError(Errors::Unknown()); // TODO
+		}
+
+		return MakeFulfilledPromise<TSharedPtr<FEOSConnectLoginCredentials>>(nullptr).GetFuture();
+	});
+}
+
+void FAuthEOSGS::ProcessSuccessfulLogin(TOnlineAsyncOp<FAuthLogin>& InAsyncOp, const EOS_ProductUserId ProductUserId)
+{
+	const FOnlineAccountIdHandle LocalUserId = CreateAccountId(ProductUserId);
+
+	UE_LOG(LogTemp, Verbose, TEXT("[FAuthEOSGS::Login] Successfully logged in as [%s]"), *ToLogString(LocalUserId));
+
+	TSharedRef<FAccountInfoEOS> AccountInfo = MakeShared<FAccountInfoEOS>();
+	AccountInfo->PlatformUserId = InAsyncOp.GetParams().PlatformUserId;
+	AccountInfo->UserId = LocalUserId;
+	AccountInfo->LoginStatus = ELoginStatus::LoggedIn;
+
+	check(!AccountInfos.Contains(LocalUserId));
+	AccountInfos.Emplace(LocalUserId, AccountInfo);
+
+	FAuthLogin::Result Result = { AccountInfo };
+	InAsyncOp.SetResult(MoveTemp(Result));
+
+	FLoginStatusChanged EventParameters;
+	EventParameters.LocalUserId = LocalUserId;
+	EventParameters.PreviousStatus = ELoginStatus::NotLoggedIn;
+	EventParameters.CurrentStatus = ELoginStatus::LoggedIn;
+	OnLoginStatusChangedEvent.Broadcast(EventParameters);
+}
+
+TOnlineAsyncOpHandle<FAuthLogout> FAuthEOSGS::Logout(FAuthLogout::Params&& Params)
+{
+	const FOnlineAccountIdHandle LocalUserId = Params.LocalUserId;
+	TOnlineAsyncOpRef<FAuthLogout> Op = GetOp<FAuthLogout>(MoveTemp(Params));
+
+	if (!ValidateOnlineId(LocalUserId))
+	{
+		Op->SetError(Errors::InvalidUser());
+		return Op->GetHandle();
+	}
+
+	Op->SetResult(FAuthLogout::Result());
+	return Op->GetHandle();
+}
+
+TOnlineResult<FAuthGetAccountByPlatformUserId> FAuthEOSGS::GetAccountByPlatformUserId(FAuthGetAccountByPlatformUserId::Params&& Params)
+{
+	TResult<FOnlineAccountIdHandle, FOnlineError> PlatformUserIdResult = GetAccountIdByPlatformUserId(Params.PlatformUserId);
+	if (PlatformUserIdResult.IsOk())
+	{
+		FAuthGetAccountByPlatformUserId::Result Result = { AccountInfos.FindChecked(PlatformUserIdResult.GetOkValue()) };
+		return TOnlineResult<FAuthGetAccountByPlatformUserId>(Result);
+	}
+	else
+	{
+		return TOnlineResult<FAuthGetAccountByPlatformUserId>(PlatformUserIdResult.GetErrorValue());
+	}
+}
+
+TOnlineResult<FAuthGetAccountByAccountId> FAuthEOSGS::GetAccountByAccountId(FAuthGetAccountByAccountId::Params&& Params)
+{
+	if (TSharedRef<FAccountInfoEOS>* const FoundAccount = AccountInfos.Find(Params.LocalUserId))
+	{
+		FAuthGetAccountByAccountId::Result Result = { *FoundAccount };
+		return TOnlineResult<FAuthGetAccountByAccountId>(Result);
+	}
+	else
+	{
+		// TODO: proper error
+		return TOnlineResult<FAuthGetAccountByAccountId>(Errors::Unknown());
+	}
+}
+
+bool FAuthEOSGS::IsLoggedIn(const FOnlineAccountIdHandle& AccountId) const
+{
+	// TODO:  More logic?
+	return AccountInfos.Contains(AccountId);
+}
+
+TResult<FOnlineAccountIdHandle, FOnlineError> FAuthEOSGS::GetAccountIdByPlatformUserId(FPlatformUserId PlatformUserId) const
+{
+	for (const TPair<FOnlineAccountIdHandle, TSharedRef<FAccountInfoEOS>>& AccountPair : AccountInfos)
+	{
+		if (AccountPair.Value->PlatformUserId == PlatformUserId)
+		{
+			TResult<FOnlineAccountIdHandle, FOnlineError> Result(AccountPair.Key);
+			return Result;
+		}
+	}
+	TResult<FOnlineAccountIdHandle, FOnlineError> Result(Errors::Unknown()); // TODO: error code
+	return Result;
+}
+
+TFuture<FOnlineAccountIdHandle> FAuthEOSGS::ResolveAccountId(const FOnlineAccountIdHandle& LocalUserId, const EOS_ProductUserId ProductUserId)
+{
+	return MakeFulfilledPromise<FOnlineAccountIdHandle>(CreateAccountId(ProductUserId)).GetFuture();
+}
+
+TFuture<TArray<FOnlineAccountIdHandle>> FAuthEOSGS::ResolveAccountIds(const FOnlineAccountIdHandle& LocalUserId, const TArray<EOS_ProductUserId>& InProductUserIds)
+{
+	TArray<FOnlineAccountIdHandle> AccountIds;
+	AccountIds.Reserve(InProductUserIds.Num());
+	for (const EOS_ProductUserId ProductUserId : InProductUserIds)
+	{
+		AccountIds.Emplace(CreateAccountId(ProductUserId));
+	}
+	return MakeFulfilledPromise<TArray<FOnlineAccountIdHandle>>(MoveTemp(AccountIds)).GetFuture();
+}
+
+TFunction<TFuture<FOnlineAccountIdHandle>(FOnlineAsyncOp& InAsyncOp, const EOS_ProductUserId& ProductUserId)> FAuthEOSGS::ResolveProductIdFn()
+{
+	return [this](FOnlineAsyncOp& InAsyncOp, const EOS_ProductUserId& ProductUserId)
+	{
+		const FOnlineAccountIdHandle* LocalUserIdPtr = InAsyncOp.Data.Get<FOnlineAccountIdHandle>(TEXT("LocalUserId"));
+		if (!ensure(LocalUserIdPtr))
+		{
+			return MakeFulfilledPromise<FOnlineAccountIdHandle>().GetFuture();
+		}
+		return ResolveAccountId(*LocalUserIdPtr, ProductUserId);
+	};
+}
+
+TFunction<TFuture<TArray<FOnlineAccountIdHandle>>(FOnlineAsyncOp& InAsyncOp, const TArray<EOS_ProductUserId>& ProductUserIds)> FAuthEOSGS::ResolveProductIdsFn()
+{
+	return [this](FOnlineAsyncOp& InAsyncOp, const TArray<EOS_ProductUserId>& ProductUserIds)
+	{
+		const FOnlineAccountIdHandle* LocalUserIdPtr = InAsyncOp.Data.Get<FOnlineAccountIdHandle>(TEXT("LocalUserId"));
+		if (!ensure(LocalUserIdPtr))
+		{
+			return MakeFulfilledPromise<TArray<FOnlineAccountIdHandle>>().GetFuture();
+		}
+		return ResolveAccountIds(*LocalUserIdPtr, ProductUserIds);
+	};
+}
+
+FOnlineAccountIdHandle FAuthEOSGS::CreateAccountId(const EOS_ProductUserId ProductUserId)
+{
+	return FOnlineAccountIdRegistryEOSGS::Get().FindOrAddAccountId(ProductUserId);
+}
+
+/* UE::Online */ }
diff --git a/Engine/Plugins/Online/Experimental/OnlineServicesEOSGS/Source/Private/Online/LobbiesEOSGS.cpp b/Engine/Plugins/Online/Experimental/OnlineServicesEOSGS/Source/Private/Online/LobbiesEOSGS.cpp
new file mode 100644
--- /dev/null
+++ b/Engine/Plugins/Online/Experimental/OnlineServicesEOSGS/Source/Private/Online/LobbiesEOSGS.cpp
@@ -0,0 +1,1763 @@
+// Copyright Epic Games, Inc. All Rights Reserved.
+
+#include "Online/LobbiesEOSGS.h"
+
+#include "Online/AuthEOSGS.h"
+#include "Online/LobbiesEOSGSTypes.h"
+#include "Online/OnlineAsyncOp.h"
+#include "Online/OnlineIdEOSGS.h"
+#include "Online/OnlineServicesEOSGS.h"
+#include "Online/OnlineServicesEOSGSTypes.h"
+
+#include "IEOSSDKManager.h"
+
+#include "eos_lobby.h"
+
+namespace UE::Online {
+
+static const FString LobbyDataKeyName = TEXT("LobbyData");
+static const FString LobbyDetailsKeyName = TEXT("LobbyDetails");
+static const FString LobbyChangesKeyName = TEXT("LobbyChanges");
+static const FString LobbyMemberChangesKeyName = TEXT("LobbyMemberChanges");
+
+static const int32 MaxAttributeSize = 1000;
+
+template <typename DataType, typename OpType>
+const DataType& GetOpDataChecked(const TOnlineAsyncOp<OpType>& Op, const FString& Key)
+{
+	const DataType* Data = Op.Data.template Get<DataType>(Key);
+	check(Data);
+	return *Data;
+}
+
+FLobbiesEOSGS::FLobbiesEOSGS(FOnlineServicesEOSGS& InServices)
+	: FLobbiesCommon(InServices)
+{
+}
+
+void FLobbiesEOSGS::Initialize()
+{
+	FLobbiesCommon::Initialize();
+
+	IEOSSDKManager* SDKManager = IEOSSDKManager::Get();
+	check(SDKManager);
+
+	EOS_HLobby LobbyInterfaceHandle = EOS_Platform_GetLobbyInterface(static_cast<FOnlineServicesEOSGS&>(GetServices()).GetEOSPlatformHandle());
+	check(LobbyInterfaceHandle != nullptr);
+
+	LobbyPrerequisites = MakeShared<FLobbyPrerequisitesEOS>(FLobbyPrerequisitesEOS{
+		LobbyInterfaceHandle,
+		StaticCastSharedPtr<FAuthEOSGS>(Services.GetAuthInterface()),
+		LobbySchemaRegistry.ToSharedRef(),
+		ServiceSchema.ToSharedRef(),
+		{SDKManager->GetProductName(), GetBuildUniqueId()}});
+
+	LobbyDataRegistry = MakeShared<FLobbyDataRegistryEOS>(LobbyPrerequisites.ToSharedRef());
+
+	RegisterHandlers();
+}
+
+void FLobbiesEOSGS::PreShutdown()
+{
+	UnregisterHandlers();
+
+	LobbyDataRegistry = nullptr;
+	LobbyPrerequisites = nullptr;
+}
+
+TOnlineAsyncOpHandle<FCreateLobby> FLobbiesEOSGS::CreateLobby(FCreateLobby::Params&& InParams)
+{
+	TOnlineAsyncOpRef<FCreateLobby> Op = GetOp<FCreateLobby>(MoveTemp(InParams));
+	const FCreateLobby::Params& Params = Op->GetParams();
+
+	if (!Services.Get<FAuthEOSGS>()->IsLoggedIn(Params.LocalUserId))
+	{
+		Op->SetError(Errors::InvalidUser());
+		return Op->GetHandle();
+	}
+
+	auto DestroyLobbyDuringCreate = [this](
+		TOnlineAsyncOp<FCreateLobby>& InAsyncOp,
+		const TSharedPtr<FLobbyDataEOS>& LobbyData,
+		FOnlineAccountIdHandle LocalUserId,
+		FOnlineError ErrorResult) -> TFuture<void>
+	{
+		FDestroyLobbyImpl::Params DestroyLobbyParams;
+		DestroyLobbyParams.LobbyData = LobbyData;
+		DestroyLobbyParams.LocalUserId = LocalUserId;
+
+		TPromise<void> Promise;
+		auto Future = Promise.GetFuture();
+
+		DestroyLobbyImpl(MoveTemp(DestroyLobbyParams))
+		.Then([this, InAsyncOp = InAsyncOp.AsShared(), ErrorResult = MoveTemp(ErrorResult), Promise = MoveTemp(Promise)](TFuture<TDefaultErrorResult<FDestroyLobbyImpl>>&& Future) mutable
+		{
+			if (Future.Get().IsError())
+			{
+				// Todo: complain about having an error while handling an error.
+			}
+
+			// Todo: Errors.
+			InAsyncOp->SetError(Errors::Unknown(MoveTemp(ErrorResult)));
+			Promise.EmplaceValue();
+		});
+
+		return Future;
+	};
+
+	Op->Then([this](TOnlineAsyncOp<FCreateLobby>& InAsyncOp)
+	{
+		// Step 1: Call create lobby.
+
+		const FCreateLobby::Params& Params = InAsyncOp.GetParams();
+		const FLobbyBucketIdTranslator<ELobbyTranslationType::ToService> BucketTanslator(LobbyPrerequisites->BucketId);
+
+		// The lobby will be created as invitation only.
+		// Once all local members are joined and the lobby attributes have been set the privacy setting will be moved to the user setting.
+
+		EOS_Lobby_CreateLobbyOptions CreateLobbyOptions = {};
+		CreateLobbyOptions.ApiVersion = EOS_LOBBY_CREATELOBBY_API_LATEST;
+		CreateLobbyOptions.LocalUserId = GetProductUserIdChecked(Params.LocalUserId);
+		CreateLobbyOptions.MaxLobbyMembers = Params.MaxMembers;
+		CreateLobbyOptions.PermissionLevel = TranslateJoinPolicy(ELobbyJoinPolicy::InvitationOnly);
+		CreateLobbyOptions.bPresenceEnabled = false; // todo: handle
+		CreateLobbyOptions.bAllowInvites = true; // todo: handle
+		CreateLobbyOptions.BucketId = BucketTanslator.GetBucketIdEOS();
+		CreateLobbyOptions.bDisableHostMigration = false; // todo: handle
+		CreateLobbyOptions.bEnableRTCRoom = false; // todo: handle
+
+		return EOS_Async<EOS_Lobby_CreateLobbyCallbackInfo>(EOS_Lobby_CreateLobby, LobbyPrerequisites->LobbyInterfaceHandle, CreateLobbyOptions); 
+	})
+	.Then([this](TOnlineAsyncOp<FCreateLobby>& InAsyncOp, const EOS_Lobby_CreateLobbyCallbackInfo* Data)
+	{
+		// Step 2: Start creating the lobby data from the EOS lobby details object.
+
+		const FCreateLobby::Params& Params = InAsyncOp.GetParams();
+
+		if (Data->ResultCode != EOS_EResult::EOS_Success)
+		{
+			// Todo: errors
+			InAsyncOp.SetError(Errors::Unknown(FromEOSError(Data->ResultCode)));
+			return MakeFulfilledPromise<TDefaultErrorResultInternal<TSharedRef<FLobbyDataEOS>>>().GetFuture();
+		}
+
+		TDefaultErrorResultInternal<TSharedRef<FLobbyDetailsEOS>> LobbyDetailsResult =
+			FLobbyDetailsEOS::CreateFromLobbyId(LobbyPrerequisites.ToSharedRef(), Params.LocalUserId, Data->LobbyId);
+		if (LobbyDetailsResult.IsError())
+		{
+			// Todo: manually call eos destroy lobby here.
+
+			// Todo: errors
+			InAsyncOp.SetError(MoveTemp(LobbyDetailsResult.GetErrorValue()));
+			return MakeFulfilledPromise<TDefaultErrorResultInternal<TSharedRef<FLobbyDataEOS>>>().GetFuture();
+		}
+		
+		return LobbyDataRegistry->FindOrCreateFromLobbyDetails(InAsyncOp.GetParams().LocalUserId, LobbyDetailsResult.GetOkValue());
+	})
+	.Then([this, DestroyLobbyDuringCreate](TOnlineAsyncOp<FCreateLobby>& InAsyncOp, TDefaultErrorResultInternal<TSharedRef<FLobbyDataEOS>>&& Result)
+	{
+		// Step 3: Store the lobby data on the async op properties.
+
+		const FCreateLobby::Params& Params = InAsyncOp.GetParams();
+
+		if (Result.IsError())
+		{
+			// Todo: destroy lobby.
+			// Change DestroyLobbyImpl to take an EOS_LobbyId instead of an FLobbyDataEOS.
+
+			// Todo: errors
+			InAsyncOp.SetError(Errors::Unknown(MoveTemp(Result.GetErrorValue())));
+			return;
+		}
+
+		// Store lobby data on the operation.
+		TSharedRef<FLobbyDataEOS> LobbyData = Result.GetOkValue();
+		InAsyncOp.Data.Set<TSharedRef<FLobbyDataEOS>>(LobbyDataKeyName, MoveTemp(LobbyData));
+	})
+	.Then([this](TOnlineAsyncOp<FCreateLobby>& InAsyncOp)
+	{
+		// Step 4: Set attributes for the lobby creator.
+
+		const FCreateLobby::Params& Params = InAsyncOp.GetParams();
+		const TSharedRef<FLobbyDataEOS>& LobbyData = GetOpDataChecked<TSharedRef<FLobbyDataEOS>>(InAsyncOp, LobbyDataKeyName);
+
+		// Add creator attributes if set.
+		// Todo: make this nicer.
+		const FJoinLobbyLocalUserData* UserData = Params.LocalUsers.FindByPredicate(
+		[LocalUserId = Params.LocalUserId](const FJoinLobbyLocalUserData& Data)
+		{
+			return LocalUserId == Data.LocalUserId;
+		});
+
+		TSharedRef<FClientLobbyMemberDataChanges> LobbyOwnerAttributes = MakeShared<FClientLobbyMemberDataChanges>();
+
+		// Add owner attributes to operation data.
+		// CreatingMemberData is used to update the local lobby data and for dispatching notifications once creation has completed.
+		TMap<FOnlineAccountIdHandle, TSharedRef<FClientLobbyMemberDataChanges>> CreatingMemberData;
+		CreatingMemberData.Reserve(Params.LocalUsers.Num());
+		CreatingMemberData.Add(Params.LocalUserId, LobbyOwnerAttributes);
+		InAsyncOp.Data.Set<TMap<FOnlineAccountIdHandle, TSharedRef<FClientLobbyMemberDataChanges>>>(LobbyMemberChangesKeyName, MoveTemp(CreatingMemberData));
+
+		if (UserData)
+		{
+			LobbyOwnerAttributes->MutatedAttributes = UserData->Attributes;
+
+			FModifyLobbyMemberDataImpl::Params ModifyLobbyMemberDataParams;
+			ModifyLobbyMemberDataParams.LobbyData = LobbyData;
+			ModifyLobbyMemberDataParams.LocalUserId = Params.LocalUserId;
+			ModifyLobbyMemberDataParams.Changes = LobbyOwnerAttributes;
+
+			return ModifyLobbyMemberDataImpl(MoveTemp(ModifyLobbyMemberDataParams));
+		}
+		else
+		{
+			return MakeFulfilledPromise<TDefaultErrorResult<FModifyLobbyMemberDataImpl>>().GetFuture();
+		}
+	})
+	.Then([this, DestroyLobbyDuringCreate](TOnlineAsyncOp<FCreateLobby>& InAsyncOp, TDefaultErrorResult<FModifyLobbyMemberDataImpl>&& ModifyLobbyOwnerResult)
+	{
+		// Step 5: Handle result
+
+		const FCreateLobby::Params& Params = InAsyncOp.GetParams();
+		TSharedRef<FLobbyDataEOS> LobbyData = GetOpDataChecked<TSharedRef<FLobbyDataEOS>>(InAsyncOp, LobbyDataKeyName);
+
+		if (ModifyLobbyOwnerResult.IsError())
+		{
+			return DestroyLobbyDuringCreate(InAsyncOp, LobbyData, Params.LocalUserId, MoveTemp(ModifyLobbyOwnerResult.GetErrorValue()));
+		}
+		else
+		{
+			return MakeFulfilledPromise<void>().GetFuture();
+		}
+	})
+	.Then([this](TOnlineAsyncOp<FCreateLobby>& InAsyncOp)
+	{
+		// Step 6: Todo: Add other local members.
+
+		// Store member attributes on the operation.
+	})
+	.Then([this](TOnlineAsyncOp<FCreateLobby>& InAsyncOp)
+	{
+		// Step 7: Add lobby attributes, set lobby join policy to user setting.
+
+		const FCreateLobby::Params& Params = InAsyncOp.GetParams();
+		const TSharedRef<FLobbyDataEOS>& LobbyData = GetOpDataChecked<TSharedRef<FLobbyDataEOS>>(InAsyncOp, LobbyDataKeyName);
+
+		// Track lobby changes to sync client-side cache at operation completion.
+		TSharedRef<FClientLobbyDataChanges> LobbyChanges = MakeShared<FClientLobbyDataChanges>();
+		InAsyncOp.Data.Set<TSharedRef<FClientLobbyDataChanges>>(LobbyChangesKeyName, LobbyChanges);
+		LobbyChanges->MutatedAttributes = Params.Attributes;
+		LobbyChanges->JoinPolicy = Params.JoinPolicy;
+
+		// Add lobby attributes.
+		// Set lobby privacy setting to the user provided value.
+		FModifyLobbyDataImpl::Params ModifyLobbyDataParams;
+		ModifyLobbyDataParams.LobbyData = LobbyData;
+		ModifyLobbyDataParams.LocalUserId = Params.LocalUserId;
+		ModifyLobbyDataParams.Changes = LobbyChanges;
+
+		return ModifyLobbyDataImpl(MoveTemp(ModifyLobbyDataParams));
+	})
+	.Then([this, DestroyLobbyDuringCreate](TOnlineAsyncOp<FCreateLobby>& InAsyncOp, TDefaultErrorResult<FModifyLobbyDataImpl>&& ModifyLobbyResult) mutable
+	{
+		// Step 8: Handle result.
+
+		const FCreateLobby::Params& Params = InAsyncOp.GetParams();
+		const TSharedRef<FLobbyDataEOS>& LobbyData = GetOpDataChecked<TSharedRef<FLobbyDataEOS>>(InAsyncOp, LobbyDataKeyName);
+
+		if (ModifyLobbyResult.IsError())
+		{
+			return DestroyLobbyDuringCreate(InAsyncOp, LobbyData, Params.LocalUserId, MoveTemp(ModifyLobbyResult.GetErrorValue()));
+		}
+		else
+		{
+			return MakeFulfilledPromise<void>().GetFuture();
+		}
+	})
+	.Then([this](TOnlineAsyncOp<FCreateLobby>& InAsyncOp)
+	{
+		// Step 9: Add the lobby to the active list for each member, apply changes to the cached
+		// lobby object, and signal notifications.
+		// The active lobbies list holds a reference to the lobby data to keep it from being cleaned up.
+
+		const FCreateLobby::Params& Params = InAsyncOp.GetParams();
+		const TSharedRef<FLobbyDataEOS>& LobbyData = GetOpDataChecked<TSharedRef<FLobbyDataEOS>>(InAsyncOp, LobbyDataKeyName);
+
+		// Mark the lobby active for each member.
+		for (const FJoinLobbyLocalUserData& LocalUser : Params.LocalUsers)
+		{
+			AddActiveLobby(LocalUser.LocalUserId, LobbyData);
+		}
+
+		// Add member changes to the lobby changes object.
+		const TSharedRef<FClientLobbyDataChanges>& LobbyChanges = GetOpDataChecked<TSharedRef<FClientLobbyDataChanges>>(InAsyncOp, LobbyChangesKeyName);
+		const TMap<FOnlineAccountIdHandle, TSharedRef<FClientLobbyMemberDataChanges>>& MemberChanges =
+			GetOpDataChecked<TMap<FOnlineAccountIdHandle, TSharedRef<FClientLobbyMemberDataChanges>>>(InAsyncOp, LobbyMemberChangesKeyName);
+		LobbyChanges->MutatedMembers = MemberChanges;
+		LobbyChanges->LocalName = InAsyncOp.GetParams().LocalName;
+
+		// Make local changes to lobby data and generate notifications.
+		LobbyData->GetClientLobbyData()->ApplyLobbyUpdateFromLocalChanges(MoveTemp(*LobbyChanges), LobbyEvents);
+
+		InAsyncOp.SetResult(FCreateLobby::Result{LobbyData->GetClientLobbyData()->GetPublicDataPtr()});
+	})
+	.Enqueue(GetSerialQueue());
+
+	return Op->GetHandle();
+}
+
+TOnlineAsyncOpHandle<FFindLobbies> FLobbiesEOSGS::FindLobbies(FFindLobbies::Params&& InParams)
+{
+	TOnlineAsyncOpRef<FFindLobbies> Op = GetOp<FFindLobbies>(MoveTemp(InParams));
+	const FFindLobbies::Params& Params = Op->GetParams();
+
+	if (!Services.Get<FAuthEOSGS>()->IsLoggedIn(Params.LocalUserId))
+	{
+		Op->SetError(Errors::InvalidUser());
+		return Op->GetHandle();
+	}
+
+	// Invalidate previous search results.
+	ActiveSearchResults.Remove(Params.LocalUserId);
+
+	Op->Then([this](TOnlineAsyncOp<FFindLobbies>& InAsyncOp)
+	{
+		return FLobbySearchEOS::Create(LobbyPrerequisites.ToSharedRef(), LobbyDataRegistry.ToSharedRef(), InAsyncOp.GetParams());
+	})
+	.Then([this](TOnlineAsyncOp<FFindLobbies>& InAsyncOp, TDefaultErrorResultInternal<TSharedRef<FLobbySearchEOS>> Result)
+	{
+		if (Result.IsError())
+		{
+			// TODO: Error codes
+			InAsyncOp.SetError(Errors::Unknown(Result.GetErrorValue()));
+		}
+		else
+		{
+			ActiveSearchResults.Add(InAsyncOp.GetParams().LocalUserId, Result.GetOkValue());
+			InAsyncOp.SetResult(FFindLobbies::Result{Result.GetOkValue()->GetLobbyResults()});
+		}
+	})
+	.Enqueue(GetSerialQueue());
+
+	return Op->GetHandle();
+}
+
+TOnlineAsyncOpHandle<FJoinLobby> FLobbiesEOSGS::JoinLobby(FJoinLobby::Params&& InParams)
+{
+	TOnlineAsyncOpRef<FJoinLobby> Op = GetOp<FJoinLobby>(MoveTemp(InParams));
+	const FJoinLobby::Params& Params = Op->GetParams();
+
+	if (!Services.Get<FAuthEOSGS>()->IsLoggedIn(Params.LocalUserId))
+	{
+		Op->SetError(Errors::InvalidUser());
+		return Op->GetHandle();
+	}
+
+	TSharedPtr<FLobbyDataEOS> LobbyData = LobbyDataRegistry->Find(Params.LobbyId);
+	if (!LobbyData)
+	{
+		Op->SetError(Errors::Unknown());
+		return Op->GetHandle();
+	}
+
+	Op->Then([this, LobbyData](TOnlineAsyncOp<FJoinLobby>& InAsyncOp)
+	{
+		// Join all users to the lobby.
+
+		const FJoinLobby::Params& Params = InAsyncOp.GetParams();
+
+		FJoinLobbyImpl::Params JoinParams;
+		JoinParams.LobbyData = LobbyData;
+		JoinParams.LocalUserId = Params.LocalUserId;
+		JoinParams.LocalName = Params.LocalName;
+		JoinParams.LocalUsers = Params.LocalUsers;
+		return JoinLobbyImpl(MoveTemp(JoinParams));
+	})
+	.Then([this](TOnlineAsyncOp<FJoinLobby>& InAsyncOp, TDefaultErrorResult<FJoinLobbyImpl> Result)
+	{
+		// Handle result.
+		if (Result.IsError())
+		{
+			// TODO: Error codes
+			InAsyncOp.SetError(Errors::Unknown(Result.GetErrorValue()));
+		}
+	})
+	.Then([this, LobbyData](TOnlineAsyncOp<FJoinLobby>& InAsyncOp)
+	{
+		// Mark the lobby active for each member.
+		// Add users to local lobby data and dispatch notifications.
+
+		const FJoinLobby::Params& Params = InAsyncOp.GetParams();
+
+		// Mark the lobby active for each member.
+		for (const FJoinLobbyLocalUserData& LocalUser : Params.LocalUsers)
+		{
+			AddActiveLobby(LocalUser.LocalUserId, LobbyData.ToSharedRef());
+		}
+
+		// todo: figure out a butter way to handle attribute parameters.
+		FClientLobbyDataChanges LobbyChanges;
+		LobbyChanges.LocalName = InAsyncOp.GetParams().LocalName;
+
+		for (const FJoinLobbyLocalUserData& LocalUser : Params.LocalUsers)
+		{
+			TSharedRef<FClientLobbyMemberDataChanges> MemberDataChanges = MakeShared<FClientLobbyMemberDataChanges>();
+			MemberDataChanges->MutatedAttributes = LocalUser.Attributes;
+			LobbyChanges.MutatedMembers.Add(LocalUser.LocalUserId, MoveTemp(MemberDataChanges));
+		}
+		LobbyData->GetClientLobbyData()->ApplyLobbyUpdateFromLocalChanges(MoveTemp(LobbyChanges), LobbyEvents);
+
+		InAsyncOp.SetResult(FJoinLobby::Result{ LobbyData->GetClientLobbyData()->GetPublicDataPtr()});
+	})
+	.Enqueue(GetSerialQueue());
+
+	return Op->GetHandle();
+}
+
+TOnlineAsyncOpHandle<FLeaveLobby> FLobbiesEOSGS::LeaveLobby(FLeaveLobby::Params&& InParams)
+{
+	TOnlineAsyncOpRef<FLeaveLobby> Op = GetOp<FLeaveLobby>(MoveTemp(InParams));
+	const FLeaveLobby::Params& Params = Op->GetParams();
+
+	if (!Services.Get<FAuthEOSGS>()->IsLoggedIn(Params.LocalUserId))
+	{
+		Op->SetError(Errors::InvalidUser());
+		return Op->GetHandle();
+	}
+
+	TSharedPtr<FLobbyDataEOS> LobbyData = LobbyDataRegistry->Find(Params.LobbyId);
+	if (!LobbyData)
+	{
+		Op->SetError(Errors::Unknown());
+		return Op->GetHandle();
+	}
+
+	Op->Then([this, LobbyData](TOnlineAsyncOp<FLeaveLobby>& InAsyncOp)
+	{
+		// Remove the user from the EOS lobby.
+
+		const FLeaveLobby::Params& Params = InAsyncOp.GetParams();
+
+		FLeaveLobbyImpl::Params LeaveParams;
+		LeaveParams.LobbyData = LobbyData;
+		LeaveParams.LocalUserId = Params.LocalUserId;
+		return LeaveLobbyImpl(MoveTemp(LeaveParams));
+	})
+	.Then([this, LobbyData](TOnlineAsyncOp<FLeaveLobby>& InAsyncOp, TDefaultErrorResult<FLeaveLobbyImpl> Result)
+	{
+		const FLeaveLobby::Params& Params = InAsyncOp.GetParams();
+
+		// Remove the user from the local lobby data and dispatch notifications.
+		FClientLobbyDataChanges LobbyChanges;
+		LobbyChanges.LeavingMembers.Add(Params.LocalUserId, ELobbyMemberLeaveReason::Left);
+		FApplyLobbyUpdateResult ApplyResult = LobbyData->GetClientLobbyData()->ApplyLobbyUpdateFromLocalChanges(MoveTemp(LobbyChanges), LobbyEvents);
+
+		// Remove the lobby from the active list for the user.
+		// The lobby data will be cleaned up once all references are removed.
+		for (FOnlineAccountIdHandle LeavingMember : ApplyResult.LeavingLocalMembers)
+		{
+			RemoveActiveLobby(LeavingMember, LobbyData.ToSharedRef());
+		}
+
+		if (Result.IsError())
+		{
+			// TODO: Error codes
+			InAsyncOp.SetError(Errors::Unknown(Result.GetErrorValue()));
+		}
+		else
+		{
+			InAsyncOp.SetResult(FLeaveLobby::Result{});
+		}
+	})
+	.Enqueue(GetSerialQueue());
+
+	return Op->GetHandle();
+}
+
+TOnlineAsyncOpHandle<FInviteLobbyMember> FLobbiesEOSGS::InviteLobbyMember(FInviteLobbyMember::Params&& InParams)
+{
+	TOnlineAsyncOpRef<FInviteLobbyMember> Op = GetOp<FInviteLobbyMember>(MoveTemp(InParams));
+	const FInviteLobbyMember::Params& Params = Op->GetParams();
+
+	if (!Services.Get<FAuthEOSGS>()->IsLoggedIn(Params.LocalUserId))
+	{
+		Op->SetError(Errors::InvalidUser());
+		return Op->GetHandle();
+	}
+
+	Op->Then([this](TOnlineAsyncOp<FInviteLobbyMember>& InAsyncOp)
+	{
+		const FInviteLobbyMember::Params& Params = InAsyncOp.GetParams();
+
+		FInviteLobbyMemberImpl::Params InviteParams;
+		InviteParams.LobbyData = LobbyDataRegistry->Find(Params.LobbyId);
+		InviteParams.LocalUserId = Params.LocalUserId;
+		InviteParams.TargetUserId = Params.TargetUserId;
+		return InviteLobbyMemberImpl(MoveTemp(InviteParams));
+	})
+	.Then([this](TOnlineAsyncOp<FInviteLobbyMember>& InAsyncOp, TDefaultErrorResult<FInviteLobbyMemberImpl> Result)
+	{
+		if (Result.IsError())
+		{
+			// TODO: Error codes
+			InAsyncOp.SetError(Errors::Unknown(Result.GetErrorValue()));
+		}
+		else
+		{
+			InAsyncOp.SetResult(FInviteLobbyMember::Result{});
+		}
+	})
+	.Enqueue(GetSerialQueue());
+
+	return Op->GetHandle();
+}
+
+TOnlineAsyncOpHandle<FDeclineLobbyInvitation> FLobbiesEOSGS::DeclineLobbyInvitation(FDeclineLobbyInvitation::Params&& InParams)
+{
+	TOnlineAsyncOpRef<FDeclineLobbyInvitation> Op = GetOp<FDeclineLobbyInvitation>(MoveTemp(InParams));
+	const FDeclineLobbyInvitation::Params& Params = Op->GetParams();
+
+	if (!Services.Get<FAuthEOSGS>()->IsLoggedIn(Params.LocalUserId))
+	{
+		Op->SetError(Errors::InvalidUser());
+		return Op->GetHandle();
+	}
+
+	Op->Then([this](TOnlineAsyncOp<FDeclineLobbyInvitation>& InAsyncOp)
+	{
+		const FDeclineLobbyInvitation::Params& Params = InAsyncOp.GetParams();
+
+		FDeclineLobbyInvitationImpl::Params DeclineParams;
+		DeclineParams.LocalUserId = Params.LocalUserId;
+		DeclineParams.LobbyId = Params.LobbyId;
+		return DeclineLobbyInvitationImpl(MoveTemp(DeclineParams));
+	})
+	.Then([this](TOnlineAsyncOp<FDeclineLobbyInvitation>& InAsyncOp, TDefaultErrorResult<FDeclineLobbyInvitationImpl> Result)
+	{
+		if (Result.IsError())
+		{
+			// TODO: Error codes
+			InAsyncOp.SetError(Errors::Unknown(Result.GetErrorValue()));
+		}
+		else
+		{
+			InAsyncOp.SetResult(FDeclineLobbyInvitation::Result{});
+		}
+	})
+	.Enqueue(GetSerialQueue());
+
+	return Op->GetHandle();
+}
+
+TOnlineAsyncOpHandle<FKickLobbyMember> FLobbiesEOSGS::KickLobbyMember(FKickLobbyMember::Params&& InParams)
+{
+	TOnlineAsyncOpRef<FKickLobbyMember> Op = GetOp<FKickLobbyMember>(MoveTemp(InParams));
+	const FKickLobbyMember::Params& Params = Op->GetParams();
+
+	if (!Services.Get<FAuthEOSGS>()->IsLoggedIn(Params.LocalUserId))
+	{
+		Op->SetError(Errors::InvalidUser());
+		return Op->GetHandle();
+	}
+
+	TSharedPtr<FLobbyDataEOS> LobbyData = LobbyDataRegistry->Find(Params.LobbyId);
+	if (!LobbyData)
+	{
+		Op->SetError(Errors::InvalidParams());
+		return Op->GetHandle();
+	}
+	Op->Data.Set<TSharedRef<FLobbyDataEOS>>(LobbyDataKeyName, LobbyData.ToSharedRef());
+
+	Op->Then([this](TOnlineAsyncOp<FKickLobbyMember>& InAsyncOp)
+	{
+		const FKickLobbyMember::Params& Params = InAsyncOp.GetParams();
+		const TSharedRef<FLobbyDataEOS>& LobbyData = GetOpDataChecked<TSharedRef<FLobbyDataEOS>>(InAsyncOp, LobbyDataKeyName);
+
+		FKickLobbyMemberImpl::Params KickParams;
+		KickParams.LobbyData = LobbyData;
+		KickParams.LocalUserId = Params.LocalUserId;
+		KickParams.TargetUserId = Params.TargetUserId;
+		return KickLobbyMemberImpl(MoveTemp(KickParams));
+	})
+	.Then([this](TOnlineAsyncOp<FKickLobbyMember>& InAsyncOp, TDefaultErrorResult<FKickLobbyMemberImpl> Result)
+	{
+		if (Result.IsError())
+		{
+			// TODO: Error codes
+			InAsyncOp.SetError(Errors::Unknown(Result.GetErrorValue()));
+		}
+	})
+	.Then([this](TOnlineAsyncOp<FKickLobbyMember>& InAsyncOp)
+	{
+		const FKickLobbyMember::Params& Params = InAsyncOp.GetParams();
+		const TSharedRef<FLobbyDataEOS>& LobbyData = GetOpDataChecked<TSharedRef<FLobbyDataEOS>>(InAsyncOp, LobbyDataKeyName);
+
+		// Update local cache and fire events.
+		FClientLobbyDataChanges LobbyChanges;
+		LobbyChanges.LeavingMembers.Add(Params.TargetUserId, ELobbyMemberLeaveReason::Kicked);
+		LobbyData->GetClientLobbyData()->ApplyLobbyUpdateFromLocalChanges(MoveTemp(LobbyChanges), LobbyEvents);
+		InAsyncOp.SetResult(FKickLobbyMember::Result{});
+	})
+	.Enqueue(GetSerialQueue());
+
+	return Op->GetHandle();
+}
+
+TOnlineAsyncOpHandle<FPromoteLobbyMember> FLobbiesEOSGS::PromoteLobbyMember(FPromoteLobbyMember::Params&& InParams)
+{
+	TOnlineAsyncOpRef<FPromoteLobbyMember> Op = GetOp<FPromoteLobbyMember>(MoveTemp(InParams));
+	const FPromoteLobbyMember::Params& Params = Op->GetParams();
+
+	if (!Services.Get<FAuthEOSGS>()->IsLoggedIn(Params.LocalUserId))
+	{
+		Op->SetError(Errors::InvalidUser());
+		return Op->GetHandle();
+	}
+
+	TSharedPtr<FLobbyDataEOS> LobbyData = LobbyDataRegistry->Find(Params.LobbyId);
+	if (!LobbyData)
+	{
+		Op->SetError(Errors::InvalidParams());
+		return Op->GetHandle();
+	}
+	Op->Data.Set<TSharedRef<FLobbyDataEOS>>(LobbyDataKeyName, LobbyData.ToSharedRef());
+
+	Op->Then([this](TOnlineAsyncOp<FPromoteLobbyMember>& InAsyncOp)
+	{
+		const FPromoteLobbyMember::Params& Params = InAsyncOp.GetParams();
+		const TSharedRef<FLobbyDataEOS>& LobbyData = GetOpDataChecked<TSharedRef<FLobbyDataEOS>>(InAsyncOp, LobbyDataKeyName);
+
+		FPromoteLobbyMemberImpl::Params PromoteParams;
+		PromoteParams.LobbyData = LobbyData;
+		PromoteParams.LocalUserId = Params.LocalUserId;
+		PromoteParams.TargetUserId = Params.TargetUserId;
+		return PromoteLobbyMemberImpl(MoveTemp(PromoteParams));
+	})
+	.Then([this](TOnlineAsyncOp<FPromoteLobbyMember>& InAsyncOp, TDefaultErrorResult<FPromoteLobbyMemberImpl> Result)
+	{
+		if (Result.IsError())
+		{
+			// TODO: Error codes
+			InAsyncOp.SetError(Errors::Unknown(Result.GetErrorValue()));
+		}
+	})
+	.Then([this](TOnlineAsyncOp<FPromoteLobbyMember>& InAsyncOp)
+	{
+		const FPromoteLobbyMember::Params& Params = InAsyncOp.GetParams();
+		const TSharedRef<FLobbyDataEOS>& LobbyData = GetOpDataChecked<TSharedRef<FLobbyDataEOS>>(InAsyncOp, LobbyDataKeyName);
+
+		// Update local cache and fire events.
+		FClientLobbyDataChanges LobbyChanges;
+		LobbyChanges.OwnerAccountId = Params.TargetUserId;
+		LobbyData->GetClientLobbyData()->ApplyLobbyUpdateFromLocalChanges(MoveTemp(LobbyChanges), LobbyEvents);
+		InAsyncOp.SetResult(FPromoteLobbyMember::Result{});
+	})
+	.Enqueue(GetSerialQueue());
+
+	return Op->GetHandle();
+}
+
+TOnlineAsyncOpHandle<FModifyLobbyJoinPolicy> FLobbiesEOSGS::ModifyLobbyJoinPolicy(FModifyLobbyJoinPolicy::Params&& InParams)
+{
+	TOnlineAsyncOpRef<FModifyLobbyJoinPolicy> Op = GetOp<FModifyLobbyJoinPolicy>(MoveTemp(InParams));
+	const FModifyLobbyJoinPolicy::Params& Params = Op->GetParams();
+
+	if (!Services.Get<FAuthEOSGS>()->IsLoggedIn(Params.LocalUserId))
+	{
+		Op->SetError(Errors::InvalidUser());
+		return Op->GetHandle();
+	}
+
+	TSharedPtr<FLobbyDataEOS> LobbyData = LobbyDataRegistry->Find(Params.LobbyId);
+	if (!LobbyData)
+	{
+		Op->SetError(Errors::InvalidParams());
+		return Op->GetHandle();
+	}
+	Op->Data.Set<TSharedRef<FLobbyDataEOS>>(LobbyDataKeyName, LobbyData.ToSharedRef());
+
+	Op->Then([this](TOnlineAsyncOp<FModifyLobbyJoinPolicy>& InAsyncOp)
+	{
+		const FModifyLobbyJoinPolicy::Params& Params = InAsyncOp.GetParams();
+		const TSharedRef<FLobbyDataEOS>& LobbyData = GetOpDataChecked<TSharedRef<FLobbyDataEOS>>(InAsyncOp, LobbyDataKeyName);
+
+		TSharedRef<FClientLobbyDataChanges> LobbyChanges = MakeShared<FClientLobbyDataChanges>();
+		InAsyncOp.Data.Set<TSharedRef<FClientLobbyDataChanges>>(LobbyChangesKeyName, LobbyChanges);
+		LobbyChanges->JoinPolicy = Params.JoinPolicy;
+
+		FModifyLobbyDataImpl::Params ModifyLobbyDataParams;
+		ModifyLobbyDataParams.LobbyData = LobbyData;
+		ModifyLobbyDataParams.LocalUserId = Params.LocalUserId;
+		ModifyLobbyDataParams.Changes = LobbyChanges;
+		return ModifyLobbyDataImpl(MoveTemp(ModifyLobbyDataParams));
+	})
+	.Then([this](TOnlineAsyncOp<FModifyLobbyJoinPolicy>& InAsyncOp, TDefaultErrorResult<FModifyLobbyDataImpl> Result)
+	{
+		if (Result.IsError())
+		{
+			// TODO: Error codes
+			InAsyncOp.SetError(Errors::Unknown(Result.GetErrorValue()));
+		}
+	})
+	.Then([this](TOnlineAsyncOp<FModifyLobbyJoinPolicy>& InAsyncOp)
+	{
+		const TSharedRef<FLobbyDataEOS>& LobbyData = GetOpDataChecked<TSharedRef<FLobbyDataEOS>>(InAsyncOp, LobbyDataKeyName);
+		const TSharedRef<FClientLobbyDataChanges>& LobbyChanges = GetOpDataChecked<TSharedRef<FClientLobbyDataChanges>>(InAsyncOp, LobbyChangesKeyName);
+
+		// Update local cache and fire events.
+		LobbyData->GetClientLobbyData()->ApplyLobbyUpdateFromLocalChanges(MoveTemp(*LobbyChanges), LobbyEvents);
+		InAsyncOp.SetResult(FModifyLobbyJoinPolicy::Result{});
+	})
+	.Enqueue(GetSerialQueue());
+
+	return Op->GetHandle();
+}
+
+TOnlineAsyncOpHandle<FModifyLobbyAttributes> FLobbiesEOSGS::ModifyLobbyAttributes(FModifyLobbyAttributes::Params&& InParams)
+{
+	TOnlineAsyncOpRef<FModifyLobbyAttributes> Op = GetOp<FModifyLobbyAttributes>(MoveTemp(InParams));
+	const FModifyLobbyAttributes::Params& Params = Op->GetParams();
+
+	if (!Services.Get<FAuthEOSGS>()->IsLoggedIn(Params.LocalUserId))
+	{
+		Op->SetError(Errors::InvalidUser());
+		return Op->GetHandle();
+	}
+
+	TSharedPtr<FLobbyDataEOS> LobbyData = LobbyDataRegistry->Find(Params.LobbyId);
+	if (!LobbyData)
+	{
+		Op->SetError(Errors::InvalidParams());
+		return Op->GetHandle();
+	}
+	Op->Data.Set<TSharedRef<FLobbyDataEOS>>(LobbyDataKeyName, LobbyData.ToSharedRef());
+
+	Op->Then([this](TOnlineAsyncOp<FModifyLobbyAttributes>& InAsyncOp)
+	{
+		const FModifyLobbyAttributes::Params& Params = InAsyncOp.GetParams();
+		const TSharedRef<FLobbyDataEOS>& LobbyData = GetOpDataChecked<TSharedRef<FLobbyDataEOS>>(InAsyncOp, LobbyDataKeyName);
+
+		TSharedRef<FClientLobbyDataChanges> LobbyChanges = MakeShared<FClientLobbyDataChanges>();
+		InAsyncOp.Data.Set<TSharedRef<FClientLobbyDataChanges>>(LobbyChangesKeyName, LobbyChanges);
+		LobbyChanges->MutatedAttributes = Params.MutatedAttributes;
+		LobbyChanges->ClearedAttributes = Params.ClearedAttributes;
+
+		FModifyLobbyDataImpl::Params ModifyLobbyDataParams;
+		ModifyLobbyDataParams.LobbyData = LobbyDataRegistry->Find(Params.LobbyId);
+		ModifyLobbyDataParams.LocalUserId = Params.LocalUserId;
+		ModifyLobbyDataParams.Changes = LobbyChanges;
+		return ModifyLobbyDataImpl(MoveTemp(ModifyLobbyDataParams));
+	})
+	.Then([this](TOnlineAsyncOp<FModifyLobbyAttributes>& InAsyncOp, TDefaultErrorResult<FModifyLobbyDataImpl> Result)
+	{
+		if (Result.IsError())
+		{
+			// TODO: Error codes
+			InAsyncOp.SetError(Errors::Unknown(Result.GetErrorValue()));
+		}
+	})
+	.Then([this](TOnlineAsyncOp<FModifyLobbyAttributes>& InAsyncOp)
+	{
+		const TSharedRef<FLobbyDataEOS>& LobbyData = GetOpDataChecked<TSharedRef<FLobbyDataEOS>>(InAsyncOp, LobbyDataKeyName);
+		const TSharedRef<FClientLobbyDataChanges>& LobbyChanges = GetOpDataChecked<TSharedRef<FClientLobbyDataChanges>>(InAsyncOp, LobbyChangesKeyName);
+
+		// Update local cache and fire events.
+		LobbyData->GetClientLobbyData()->ApplyLobbyUpdateFromLocalChanges(MoveTemp(*LobbyChanges), LobbyEvents);
+		InAsyncOp.SetResult(FModifyLobbyAttributes::Result{});
+	})
+	.Enqueue(GetSerialQueue());
+
+	return Op->GetHandle();
+}
+
+TOnlineAsyncOpHandle<FModifyLobbyMemberAttributes> FLobbiesEOSGS::ModifyLobbyMemberAttributes(FModifyLobbyMemberAttributes::Params&& InParams)
+{
+	TOnlineAsyncOpRef<FModifyLobbyMemberAttributes> Op = GetOp<FModifyLobbyMemberAttributes>(MoveTemp(InParams));
+	const FModifyLobbyMemberAttributes::Params& Params = Op->GetParams();
+
+	if (!Services.Get<FAuthEOSGS>()->IsLoggedIn(Params.LocalUserId))
+	{
+		Op->SetError(Errors::InvalidUser());
+		return Op->GetHandle();
+	}
+
+	TSharedPtr<FLobbyDataEOS> LobbyData = LobbyDataRegistry->Find(Params.LobbyId);
+	if (!LobbyData)
+	{
+		Op->SetError(Errors::InvalidParams());
+		return Op->GetHandle();
+	}
+	Op->Data.Set<TSharedRef<FLobbyDataEOS>>(LobbyDataKeyName, LobbyData.ToSharedRef());
+
+	Op->Then([this](TOnlineAsyncOp<FModifyLobbyMemberAttributes>& InAsyncOp)
+	{
+		const FModifyLobbyMemberAttributes::Params& Params = InAsyncOp.GetParams();
+		const TSharedRef<FLobbyDataEOS>& LobbyData = GetOpDataChecked<TSharedRef<FLobbyDataEOS>>(InAsyncOp, LobbyDataKeyName);
+
+		TSharedRef<FClientLobbyMemberDataChanges> LobbyMemberChanges = MakeShared<FClientLobbyMemberDataChanges>();
+		InAsyncOp.Data.Set<TSharedRef<FClientLobbyMemberDataChanges>>(LobbyMemberChangesKeyName, LobbyMemberChanges);
+		LobbyMemberChanges->MutatedAttributes = Params.MutatedAttributes;
+		LobbyMemberChanges->ClearedAttributes = Params.ClearedAttributes;
+
+		FModifyLobbyMemberDataImpl::Params ModifyLobbyMemberDataParams;
+		ModifyLobbyMemberDataParams.LobbyData = LobbyData;
+		ModifyLobbyMemberDataParams.LocalUserId = Params.LocalUserId;
+		ModifyLobbyMemberDataParams.Changes = LobbyMemberChanges;
+		return ModifyLobbyMemberDataImpl(MoveTemp(ModifyLobbyMemberDataParams));
+	})
+	.Then([this](TOnlineAsyncOp<FModifyLobbyMemberAttributes>& InAsyncOp, TDefaultErrorResult<FModifyLobbyMemberDataImpl> Result)
+	{
+		if (Result.IsError())
+		{
+			// TODO: Error codes
+			InAsyncOp.SetError(Errors::Unknown(Result.GetErrorValue()));
+		}
+	})
+	.Then([this](TOnlineAsyncOp<FModifyLobbyMemberAttributes>& InAsyncOp)
+	{
+		const FModifyLobbyMemberAttributes::Params& Params = InAsyncOp.GetParams();
+		const TSharedRef<FLobbyDataEOS>& LobbyData = GetOpDataChecked<TSharedRef<FLobbyDataEOS>>(InAsyncOp, LobbyDataKeyName);
+		const TSharedRef<FClientLobbyMemberDataChanges>& LobbyMemberChanges = GetOpDataChecked<TSharedRef<FClientLobbyMemberDataChanges>>(InAsyncOp, LobbyMemberChangesKeyName);
+
+		// Update local cache and fire events.
+		FClientLobbyDataChanges LobbyChanges;
+		LobbyChanges.MutatedMembers.Add(Params.LocalUserId, LobbyMemberChanges);
+		LobbyData->GetClientLobbyData()->ApplyLobbyUpdateFromLocalChanges(MoveTemp(LobbyChanges), LobbyEvents);
+		InAsyncOp.SetResult(FModifyLobbyMemberAttributes::Result{});
+	})
+	.Enqueue(GetSerialQueue());
+
+	return Op->GetHandle();
+}
+
+TOnlineResult<FGetJoinedLobbies> FLobbiesEOSGS::GetJoinedLobbies(FGetJoinedLobbies::Params&& Params)
+{
+	if (const TSet<TSharedRef<FLobbyDataEOS>>* Lobbies = ActiveLobbies.Find(Params.LocalUserId))
+	{
+		FGetJoinedLobbies::Result Result;
+		Result.Lobbies.Reserve(Lobbies->Num());
+		for (const TSharedRef<FLobbyDataEOS>& LobbyDataEOS : *Lobbies)
+		{
+			Result.Lobbies.Emplace(LobbyDataEOS->GetClientLobbyData()->GetPublicDataPtr());
+		}
+		return TOnlineResult<FGetJoinedLobbies>(MoveTemp(Result));
+	}
+	else
+	{
+		return TOnlineResult<FGetJoinedLobbies>(Errors::InvalidUser());
+	}
+}
+
+void FLobbiesEOSGS::HandleLobbyUpdated(const EOS_Lobby_LobbyUpdateReceivedCallbackInfo* Data)
+{
+	if (TSharedPtr<FLobbyDataEOS> LobbyData = LobbyDataRegistry->Find(Data->LobbyId))
+	{
+		FProcessLobbyNotificationImpl::Params Params;
+		Params.LobbyData = LobbyData;
+
+		ProcessLobbyNotificationImplOp(MoveTemp(Params))
+		.OnComplete([LobbyData](const TOnlineResult<FProcessLobbyNotificationImpl>& Result)
+		{
+			if (Result.IsError())
+			{
+				// Todo: handle failure to update lobby from snapshot.
+				UE_LOG(LogTemp, Warning, TEXT("[FLobbiesEOSGS::HandleLobbyUpdated] Failed to apply update. Lobby: %s, Error: %s"), *LobbyData->GetLobbyId(), *Result.GetErrorValue().GetLogString());
+			}
+		});
+	}
+}
+
+void FLobbiesEOSGS::HandleLobbyMemberUpdated(const EOS_Lobby_LobbyMemberUpdateReceivedCallbackInfo* Data)
+{
+	if (TSharedPtr<FLobbyDataEOS> LobbyData = LobbyDataRegistry->Find(Data->LobbyId))
+	{
+		FProcessLobbyNotificationImpl::Params Params;
+		Params.LobbyData = LobbyData;
+		Params.MutatedMembers.Add(Data->TargetUserId);
+
+		ProcessLobbyNotificationImplOp(MoveTemp(Params))
+		.OnComplete([LobbyData](const TOnlineResult<FProcessLobbyNotificationImpl>& Result)
+		{
+			if (Result.IsError())
+			{
+				// Todo: handle failure to update lobby from snapshot.
+				UE_LOG(LogTemp, Warning, TEXT("[FLobbiesEOSGS::HandleLobbyMemberUpdated] Failed to apply update. Lobby: %s, Error: %s"), *LobbyData->GetLobbyId(), *Result.GetErrorValue().GetLogString());
+			}
+		});
+	}
+}
+
+void FLobbiesEOSGS::HandleLobbyMemberStatusReceived(const EOS_Lobby_LobbyMemberStatusReceivedCallbackInfo* Data)
+{
+	if (TSharedPtr<FLobbyDataEOS> LobbyData = LobbyDataRegistry->Find(Data->LobbyId))
+	{
+		FProcessLobbyNotificationImpl::Params Params;
+		Params.LobbyData = LobbyData;
+
+		switch (Data->CurrentStatus)
+		{
+		case EOS_ELobbyMemberStatus::EOS_LMS_JOINED:
+			// Fetch member snapshot on join.
+			Params.MutatedMembers.Add(Data->TargetUserId);
+			break;
+
+		case EOS_ELobbyMemberStatus::EOS_LMS_PROMOTED:
+			// No member data needed, only lobby snapshot.
+			break;
+
+		case EOS_ELobbyMemberStatus::EOS_LMS_LEFT:
+			Params.LeavingMembers.Add(Data->TargetUserId, ELobbyMemberLeaveReason::Left);
+			break;
+
+		case EOS_ELobbyMemberStatus::EOS_LMS_DISCONNECTED:
+			Params.LeavingMembers.Add(Data->TargetUserId, ELobbyMemberLeaveReason::Disconnected);
+			break;
+
+		case EOS_ELobbyMemberStatus::EOS_LMS_KICKED:
+			Params.LeavingMembers.Add(Data->TargetUserId, ELobbyMemberLeaveReason::Kicked);
+			break;
+
+		default: checkNoEntry(); // Intentional fallthrough
+		case EOS_ELobbyMemberStatus::EOS_LMS_CLOSED:
+			Params.LeavingMembers.Add(Data->TargetUserId, ELobbyMemberLeaveReason::Closed);
+			break;
+		}
+
+		ProcessLobbyNotificationImplOp(MoveTemp(Params))
+		.OnComplete([LobbyData](const TOnlineResult<FProcessLobbyNotificationImpl>& Result)
+		{
+			if (Result.IsError())
+			{
+				// Todo: handle failure to update lobby from snapshot.
+				UE_LOG(LogTemp, Warning, TEXT("[FLobbiesEOSGS::HandleLobbyMemberUpdated] Failed to apply update. Lobby: %s, Error: %s"), *LobbyData->GetLobbyId(), *Result.GetErrorValue().GetLogString());
+			}
+		});
+	}
+}
+
+void FLobbiesEOSGS::HandleLobbyInviteReceived(const EOS_Lobby_LobbyInviteReceivedCallbackInfo* Data)
+{
+#if 0
+	EOS_STRUCT(EOS_Lobby_LobbyInviteReceivedCallbackInfo, (
+		/** Context that was passed into EOS_Lobby_AddNotifyLobbyInviteReceived */
+		void* ClientData;
+		/** The ID of the invitation */
+		const char* InviteId;
+		/** The Product User ID of the local user who received the invitation */
+		EOS_ProductUserId LocalUserId;
+		/** The Product User ID of the user who sent the invitation */
+		EOS_ProductUserId TargetUserId;
+	));
+#endif
+
+	// Todo: Queue this like an operation.
+	const FOnlineAccountIdHandle LocalUserId = FindAccountId(Data->LocalUserId);
+	if (LocalUserId.IsValid())
+	{
+		FLobbyInviteDataEOS::CreateFromInviteId(LobbyPrerequisites.ToSharedRef(), LobbyDataRegistry.ToSharedRef(), LocalUserId, Data->InviteId, Data->TargetUserId)
+		.Then([this](TFuture<TDefaultErrorResultInternal<TSharedRef<FLobbyInviteDataEOS>>>&& Future)
+		{
+			if (Future.Get().IsError())
+			{
+				// Todo: Log / queue a manual fetch of invitations.
+				UE_LOG(LogTemp, Warning, TEXT("[FLobbiesEOSGS::HandleLobbyInviteReceived] Failed to receive invite. Error: %s"),
+					*Future.Get().GetErrorValue().GetLogString());
+			}
+			else
+			{
+				UE_LOG(LogTemp, Warning, TEXT("[FLobbiesEOSGS::HandleLobbyInviteReceived] Received invite. Id: %s, Lobby: %s, Receiver: %s, Sender: %s"),
+					*Future.Get().GetOkValue()->GetInviteId(),
+					*Future.Get().GetOkValue()->GetLobbyData()->GetLobbyId(),
+					*ToLogString(Future.Get().GetOkValue()->GetReceiver()),
+					*ToLogString(Future.Get().GetOkValue()->GetSender()));
+				AddActiveInvite(Future.Get().GetOkValue());
+			}
+		});
+	}
+}
+
+void FLobbiesEOSGS::HandleLobbyInviteAccepted(const EOS_Lobby_LobbyInviteAcceptedCallbackInfo* Data)
+{
+#if 0
+	EOS_STRUCT(EOS_Lobby_LobbyInviteAcceptedCallbackInfo, (
+		/** Context that was passed into EOS_Lobby_AddNotifyLobbyInviteAccepted */
+		void* ClientData;
+		/** The invite ID */
+		const char* InviteId;
+		/** The Product User ID of the local user who received the invitation */
+		EOS_ProductUserId LocalUserId;
+		/** The Product User ID of the user who sent the invitation */
+		EOS_ProductUserId TargetUserId;
+		/** Lobby ID that the user has been invited to */
+		const char* LobbyId;
+	));
+#endif
+
+	// Todo: handle catalog of sent invitations.
+}
+
+void FLobbiesEOSGS::HandleJoinLobbyAccepted(const EOS_Lobby_JoinLobbyAcceptedCallbackInfo* Data)
+{
+#if 0
+	EOS_STRUCT(EOS_Lobby_JoinLobbyAcceptedCallbackInfo, (
+		/** Context that was passed into EOS_Lobby_AddNotifyJoinLobbyAccepted */
+		void* ClientData;
+		/** The Product User ID of the local user who is joining */
+		EOS_ProductUserId LocalUserId;
+		/** 
+		 * The UI Event associated with this Join Game event.
+		 * This should be used with EOS_Lobby_CopyLobbyDetailsHandleByUiEventId to get a handle to be used
+		 * when calling EOS_Lobby_JoinLobby.
+		 */
+		EOS_UI_EventId UiEventId;
+	));
+#endif
+
+	// Todo: handle UI events.
+}
+
+void FLobbiesEOSGS::RegisterHandlers()
+{
+	// Register for lobby updates.
+	OnLobbyUpdatedEOSEventRegistration = EOS_RegisterComponentEventHandler(
+		this,
+		LobbyPrerequisites->LobbyInterfaceHandle,
+		EOS_LOBBY_ADDNOTIFYLOBBYUPDATERECEIVED_API_LATEST,
+		&EOS_Lobby_AddNotifyLobbyUpdateReceived,
+		&EOS_Lobby_RemoveNotifyLobbyUpdateReceived,
+		&FLobbiesEOSGS::HandleLobbyUpdated);
+
+	// Register for lobby member updates.
+	OnLobbyMemberUpdatedEOSEventRegistration = EOS_RegisterComponentEventHandler(
+		this,
+		LobbyPrerequisites->LobbyInterfaceHandle,
+		EOS_LOBBY_ADDNOTIFYLOBBYMEMBERUPDATERECEIVED_API_LATEST,
+		&EOS_Lobby_AddNotifyLobbyMemberUpdateReceived,
+		&EOS_Lobby_RemoveNotifyLobbyMemberUpdateReceived,
+		&FLobbiesEOSGS::HandleLobbyMemberUpdated);
+
+	// Register for lobby member status changed.
+	OnLobbyMemberStatusReceivedEOSEventRegistration = EOS_RegisterComponentEventHandler(
+		this,
+		LobbyPrerequisites->LobbyInterfaceHandle,
+		EOS_LOBBY_ADDNOTIFYLOBBYMEMBERSTATUSRECEIVED_API_LATEST,
+		&EOS_Lobby_AddNotifyLobbyMemberStatusReceived,
+		&EOS_Lobby_RemoveNotifyLobbyMemberStatusReceived,
+		&FLobbiesEOSGS::HandleLobbyMemberStatusReceived);
+
+	// Register for lobby invite received.
+	OnLobbyInviteReceivedEOSEventRegistration = EOS_RegisterComponentEventHandler(
+		this,
+		LobbyPrerequisites->LobbyInterfaceHandle,
+		EOS_LOBBY_ADDNOTIFYLOBBYINVITERECEIVED_API_LATEST,
+		&EOS_Lobby_AddNotifyLobbyInviteReceived,
+		&EOS_Lobby_RemoveNotifyLobbyInviteReceived,
+		&FLobbiesEOSGS::HandleLobbyInviteReceived);
+
+	// Register for lobby invite accepted.
+	OnLobbyInviteAcceptedEOSEventRegistration = EOS_RegisterComponentEventHandler(
+		this,
+		LobbyPrerequisites->LobbyInterfaceHandle,
+		EOS_LOBBY_ADDNOTIFYLOBBYINVITERECEIVED_API_LATEST,
+		&EOS_Lobby_AddNotifyLobbyInviteAccepted,
+		&EOS_Lobby_RemoveNotifyLobbyInviteAccepted,
+		&FLobbiesEOSGS::HandleLobbyInviteAccepted);
+
+	// Register for join lobby accepted via overlay.
+	OnJoinLobbyAcceptedEOSEventRegistration = EOS_RegisterComponentEventHandler(
+		this,
+		LobbyPrerequisites->LobbyInterfaceHandle,
+		EOS_LOBBY_ADDNOTIFYJOINLOBBYACCEPTED_API_LATEST,
+		&EOS_Lobby_AddNotifyJoinLobbyAccepted,
+		&EOS_Lobby_RemoveNotifyJoinLobbyAccepted,
+		&FLobbiesEOSGS::HandleJoinLobbyAccepted);
+}
+
+void FLobbiesEOSGS::UnregisterHandlers()
+{
+	OnLobbyUpdatedEOSEventRegistration = nullptr;
+	OnLobbyMemberUpdatedEOSEventRegistration = nullptr;
+	OnLobbyMemberStatusReceivedEOSEventRegistration = nullptr;
+	OnLobbyInviteReceivedEOSEventRegistration = nullptr;
+	OnLobbyInviteAcceptedEOSEventRegistration = nullptr;
+	OnJoinLobbyAcceptedEOSEventRegistration = nullptr;
+}
+
+void FLobbiesEOSGS::AddActiveLobby(FOnlineAccountIdHandle LocalUserId, const TSharedRef<FLobbyDataEOS>& LobbyData)
+{
+	// Add bookkeeping for the user.
+	ActiveLobbies.FindOrAdd(LocalUserId).Add(LobbyData);
+}
+
+void FLobbiesEOSGS::RemoveActiveLobby(FOnlineAccountIdHandle LocalUserId, const TSharedRef<FLobbyDataEOS>& LobbyData)
+{
+	// Remove bookkeeping for the local user.
+	if (TSet<TSharedRef<FLobbyDataEOS>>* Lobbies = ActiveLobbies.Find(LocalUserId))
+	{
+		Lobbies->Remove(LobbyData);
+	}
+}
+
+void FLobbiesEOSGS::AddActiveInvite(const TSharedRef<FLobbyInviteDataEOS>& Invite)
+{
+	TMap<FOnlineLobbyIdHandle, TSharedRef<FLobbyInviteDataEOS>>& ActiveUserInvites = ActiveInvites.FindOrAdd(Invite->GetReceiver());
+	const FOnlineLobbyIdHandle LobbyId = Invite->GetLobbyData()->GetLobbyIdHandle();
+
+	// Todo: Handle multiple invites for the same lobby.
+	if (ActiveUserInvites.Find(LobbyId) == nullptr)
+	{
+		ActiveUserInvites.Add(LobbyId, Invite);
+		LobbyEvents.OnLobbyInvitationAdded.Broadcast(
+			FLobbyInvitationAdded{
+				Invite->GetReceiver(),
+				Invite->GetSender(),
+				Invite->GetLobbyData()->GetClientLobbyData()->GetPublicDataPtr()
+			});
+	}
+}
+
+void FLobbiesEOSGS::RemoveActiveInvite(const TSharedRef<FLobbyInviteDataEOS>& Invite)
+{
+	ActiveInvites.FindOrAdd(Invite->GetReceiver()).Remove(Invite->GetLobbyData()->GetLobbyIdHandle());
+
+	LobbyEvents.OnLobbyInvitationRemoved.Broadcast(
+		FLobbyInvitationRemoved{
+			Invite->GetReceiver(),
+			Invite->GetSender(),
+			Invite->GetLobbyData()->GetClientLobbyData()->GetPublicDataPtr()
+		});
+}
+
+TSharedPtr<FLobbyInviteDataEOS> FLobbiesEOSGS::GetActiveInvite(FOnlineAccountIdHandle TargetUser, FOnlineLobbyIdHandle TargetLobbyId)
+{
+	const TSharedRef<FLobbyInviteDataEOS>* Result = ActiveInvites.FindOrAdd(TargetUser).Find(TargetLobbyId);
+	return Result ? *Result : TSharedPtr<FLobbyInviteDataEOS>();
+}
+
+TFuture<TDefaultErrorResult<FLobbiesEOSGS::FJoinLobbyImpl>> FLobbiesEOSGS::JoinLobbyImpl(FJoinLobbyImpl::Params&& Params)
+{
+	// Check prerequisites.
+	if (!Services.Get<FAuthEOSGS>()->IsLoggedIn(Params.LocalUserId))
+	{
+		return MakeFulfilledPromise<TDefaultErrorResult<FJoinLobbyImpl>>(Errors::InvalidUser()).GetFuture();
+	}
+
+	if (!Params.LobbyData.IsValid())
+	{
+		return MakeFulfilledPromise<TDefaultErrorResult<FJoinLobbyImpl>>(Errors::InvalidParams()).GetFuture();
+	}
+
+	// Check whether any of the local users is already in the target lobby.
+	for (const FJoinLobbyLocalUserData& JoinData : Params.LocalUsers)
+	{
+		if (const TSet<TSharedRef<FLobbyDataEOS>>* UserActiveLobbies = ActiveLobbies.Find(JoinData.LocalUserId))
+		{
+			if (UserActiveLobbies->Find(Params.LobbyData.ToSharedRef()))
+			{
+				return MakeFulfilledPromise<TDefaultErrorResult<FJoinLobbyImpl>>(Errors::InvalidParams()).GetFuture();
+			}
+		}
+	}
+
+	TArray<TFuture<TDefaultErrorResult<FJoinLobbyMemberImpl>>> PendingMemberJoins;
+	for (FJoinLobbyLocalUserData& UserData : Params.LocalUsers)
+	{
+		FJoinLobbyMemberImpl::Params JoinLobbyMemberParams;
+		JoinLobbyMemberParams.LobbyData = Params.LobbyData.ToSharedRef();
+		JoinLobbyMemberParams.LocalUserId = UserData.LocalUserId;
+		JoinLobbyMemberParams.Attributes = MoveTemp(UserData.Attributes);
+
+		TSharedRef<TPromise<TDefaultErrorResult<FJoinLobbyMemberImpl>>> JoinMemberPromise = MakeShared<TPromise<TDefaultErrorResult<FJoinLobbyMemberImpl>>>();
+		JoinLobbyMemberImplOp(MoveTemp(JoinLobbyMemberParams))
+		.OnComplete([JoinMemberPromise](const TOnlineResult<FJoinLobbyMemberImpl>& Result)
+		{
+			if (Result.IsOk())
+			{
+				JoinMemberPromise->EmplaceValue(Result.GetOkValue());
+			}
+			else
+			{
+				JoinMemberPromise->EmplaceValue(Result.GetErrorValue());
+			}
+		});
+
+		PendingMemberJoins.Emplace(JoinMemberPromise->GetFuture());
+	}
+
+	TPromise<TDefaultErrorResult<FJoinLobbyImpl>> Promise;
+	TFuture<TDefaultErrorResult<FJoinLobbyImpl>> Future = Promise.GetFuture();
+
+	WhenAll(MoveTemp(PendingMemberJoins))
+	.Then([this, Promise = MoveTemp(Promise), Params = MoveTemp(Params)](TFuture<TArray<TDefaultErrorResult<FJoinLobbyMemberImpl>>>&& Results) mutable
+	{
+		TOptional<FOnlineError> StoredError;
+		for (TDefaultErrorResult<FJoinLobbyMemberImpl>& Result : Results.Get())
+		{
+			if (Result.IsError())
+			{
+				// Store first encountered error to return as result.
+				StoredError = Result.GetErrorValue();
+			}
+		}
+
+		if (StoredError)
+		{
+			TArray<TFuture<TDefaultErrorResult<FLeaveLobbyImpl>>> PendingMemberExits;
+			for (int32 MemberIndex = 0; MemberIndex < Params.LocalUsers.Num(); ++MemberIndex)
+			{
+				const FOnlineAccountIdHandle MemberId = Params.LocalUsers[MemberIndex].LocalUserId;
+
+				if (Results.Get()[MemberIndex].IsError())
+				{
+					FLeaveLobbyImpl::Params LeaveLobbyParams;
+					LeaveLobbyParams.LobbyData = Params.LobbyData.ToSharedRef();
+					LeaveLobbyParams.LocalUserId = MemberId;
+
+					TPromise<TDefaultErrorResult<FLeaveLobbyImpl>> LeaveMemberPromise;
+					PendingMemberExits.Emplace(LeaveMemberPromise.GetFuture());
+
+					LeaveLobbyImpl(MoveTemp(LeaveLobbyParams))
+					.Then([LeaveMemberPromise = MoveTemp(LeaveMemberPromise)](TFuture<TDefaultErrorResult<FLeaveLobbyImpl>>&& Future) mutable
+					{
+						if (Future.Get().IsError())
+						{
+							// Todo: complain about having an error while handling an error.
+						}
+
+						LeaveMemberPromise.EmplaceValue(MoveTempIfPossible(Future.Get()));
+					});
+				}
+			}
+
+			WhenAll(MoveTemp(PendingMemberExits))
+			.Then([Promise = MoveTemp(Promise), StoredError = MoveTemp(*StoredError)](TFuture<TArray<TDefaultErrorResult<FLeaveLobbyImpl>>>&& Future) mutable
+			{
+				Promise.EmplaceValue(MoveTemp(StoredError));
+			});
+		}
+		else
+		{
+			Promise.EmplaceValue(FJoinLobbyImpl::Result{});
+		}
+	});
+
+	return Future;
+}
+
+TOnlineAsyncOpHandle<FLobbiesEOSGS::FJoinLobbyMemberImpl> FLobbiesEOSGS::JoinLobbyMemberImplOp(FJoinLobbyMemberImpl::Params&& InParams)
+{
+	TOnlineAsyncOpRef<FJoinLobbyMemberImpl> Op = GetOp<FJoinLobbyMemberImpl>(MoveTemp(InParams));
+	const FJoinLobbyMemberImpl::Params& Params = Op->GetParams();
+
+	// Setup lobby details - Prefer UI event before invitation before search result.
+	TSharedPtr<FLobbyDetailsEOS> LobbyDetails = Params.LobbyData->GetUserLobbyDetails(Params.LocalUserId);
+	if (!LobbyDetails)
+	{
+		// Todo: Check whether another local member can invite the user.
+		Op->SetError(Errors::InvalidParams());
+		return Op->GetHandle();
+	}
+
+	if (LobbyDetails->GetInfo()->GetProductVersion() != LobbyPrerequisites->BucketId.GetProductVersion())
+	{
+		Op->SetError(Errors::IncompatibleVersion());
+		return Op->GetHandle();
+	}
+
+	Op->Then([this, LobbyDetails, LobbyData = Params.LobbyData](TOnlineAsyncOp<FJoinLobbyMemberImpl>& InAsyncOp)
+	{
+		const FLobbiesEOSGS::FJoinLobbyMemberImpl::Params& Params = InAsyncOp.GetParams();
+
+		EOS_Lobby_JoinLobbyOptions JoinLobbyOptions = {};
+		JoinLobbyOptions.ApiVersion = EOS_LOBBY_JOINLOBBY_API_LATEST;
+		JoinLobbyOptions.LobbyDetailsHandle = LobbyDetails->GetEOSHandle();
+		JoinLobbyOptions.LocalUserId = GetProductUserIdChecked(Params.LocalUserId);
+		JoinLobbyOptions.bPresenceEnabled = false;
+		JoinLobbyOptions.LocalRTCOptions = nullptr;
+		return EOS_Async<EOS_Lobby_JoinLobbyCallbackInfo>(EOS_Lobby_JoinLobby, LobbyPrerequisites->LobbyInterfaceHandle, JoinLobbyOptions);
+	})
+	.Then([this](TOnlineAsyncOp<FJoinLobbyMemberImpl>& InAsyncOp, const EOS_Lobby_JoinLobbyCallbackInfo* Data)
+	{
+		if (Data->ResultCode != EOS_EResult::EOS_Success)
+		{
+			// TODO: Error codes
+			InAsyncOp.SetError(Errors::Unknown(FromEOSError(Data->ResultCode)));
+		}
+	})
+	.Then([this](TOnlineAsyncOp<FJoinLobbyMemberImpl>& InAsyncOp)
+	{
+		const FLobbiesEOSGS::FJoinLobbyMemberImpl::Params& Params = InAsyncOp.GetParams();
+		FModifyLobbyMemberDataImpl::Params ModifyParams;
+		ModifyParams.LobbyData = Params.LobbyData;
+		ModifyParams.LocalUserId = Params.LocalUserId;
+		ModifyParams.Changes = MakeShared<FClientLobbyMemberDataChanges>();
+		ModifyParams.Changes->MutatedAttributes = Params.Attributes;
+		return ModifyLobbyMemberDataImpl(MoveTemp(ModifyParams));
+	})
+	.Then([this](TOnlineAsyncOp<FJoinLobbyMemberImpl>& InAsyncOp, TDefaultErrorResult<FModifyLobbyMemberDataImpl>&& Result)
+	{
+		const FLobbiesEOSGS::FJoinLobbyMemberImpl::Params& Params = InAsyncOp.GetParams();
+
+		if (Result.IsError())
+		{
+			// Failed to set attributes - leave the lobby.
+			FLeaveLobbyImpl::Params LeaveLobbyParams;
+			LeaveLobbyParams.LobbyData = Params.LobbyData;
+			LeaveLobbyParams.LocalUserId = Params.LocalUserId;
+
+			TPromise<void> Promise;
+			TFuture<void> Future = Promise.GetFuture();
+
+			LeaveLobbyImpl(MoveTemp(LeaveLobbyParams))
+			.Then([this, InAsyncOp = InAsyncOp.AsShared(), ErrorResult = MoveTemp(Result.GetErrorValue()), Promise = MoveTemp(Promise)](TFuture<TDefaultErrorResult<FLeaveLobbyImpl>>&& Future) mutable
+			{
+				if (Future.Get().IsError())
+				{
+					// Todo: complain about having an error while handling an error.
+				}
+
+				// Todo: Errors.
+				InAsyncOp->SetError(Errors::Unknown(MoveTemp(ErrorResult)));
+				Promise.EmplaceValue();
+			});
+
+			return Future;
+		}
+		else
+		{
+			return MakeFulfilledPromise<void>().GetFuture();
+		}
+	})
+	.Then([this](TOnlineAsyncOp<FJoinLobbyMemberImpl>& InAsyncOp)
+	{
+		InAsyncOp.SetResult(FJoinLobbyMemberImpl::Result{});
+	})
+	.Enqueue(GetSerialQueue(InParams.LocalUserId));
+
+	return Op->GetHandle();
+}
+
+TFuture<TDefaultErrorResult<FLobbiesEOSGS::FLeaveLobbyImpl>> FLobbiesEOSGS::LeaveLobbyImpl(FLeaveLobbyImpl::Params&& Params)
+{
+	// Check prerequisites.
+	if (!Services.Get<FAuthEOSGS>()->IsLoggedIn(Params.LocalUserId))
+	{
+		return MakeFulfilledPromise<TDefaultErrorResult<FLeaveLobbyImpl>>(Errors::InvalidUser()).GetFuture();
+	}
+
+	if (!Params.LobbyData.IsValid())
+	{
+		return MakeFulfilledPromise<TDefaultErrorResult<FLeaveLobbyImpl>>(Errors::InvalidParams()).GetFuture();
+	}
+
+	EOS_Lobby_LeaveLobbyOptions LeaveLobbyOptions = {};
+	LeaveLobbyOptions.ApiVersion = EOS_LOBBY_LEAVELOBBY_API_LATEST;
+	LeaveLobbyOptions.LocalUserId = GetProductUserIdChecked(Params.LocalUserId);
+	LeaveLobbyOptions.LobbyId = Params.LobbyData->GetLobbyIdEOS();
+
+	TPromise<TDefaultErrorResult<FLeaveLobbyImpl>> Promise;
+	TFuture<TDefaultErrorResult<FLeaveLobbyImpl>> Future = Promise.GetFuture();
+
+	EOS_Async<EOS_Lobby_LeaveLobbyCallbackInfo>(EOS_Lobby_LeaveLobby, LobbyPrerequisites->LobbyInterfaceHandle, LeaveLobbyOptions)
+	.Then([Promise = MoveTemp(Promise)](TFuture<const EOS_Lobby_LeaveLobbyCallbackInfo*>&& Future) mutable
+	{
+		const EOS_Lobby_LeaveLobbyCallbackInfo* Result = Future.Get();
+		if (Result->ResultCode != EOS_EResult::EOS_Success)
+		{
+			// Todo: Errors
+			Promise.EmplaceValue(Errors::Unknown(FromEOSError(Result->ResultCode)));
+			return;
+		}
+
+		Promise.EmplaceValue(FLeaveLobbyImpl::Result{});
+	});
+
+	return Future;
+}
+
+TFuture<TDefaultErrorResult<FLobbiesEOSGS::FDestroyLobbyImpl>> FLobbiesEOSGS::DestroyLobbyImpl(FDestroyLobbyImpl::Params&& Params)
+{
+	// Check prerequisites.
+	if (!Services.Get<FAuthEOSGS>()->IsLoggedIn(Params.LocalUserId))
+	{
+		return MakeFulfilledPromise<TDefaultErrorResult<FDestroyLobbyImpl>>(Errors::InvalidUser()).GetFuture();
+	}
+
+	if (!Params.LobbyData.IsValid())
+	{
+		return MakeFulfilledPromise<TDefaultErrorResult<FDestroyLobbyImpl>>(Errors::InvalidParams()).GetFuture();
+	}
+
+	EOS_Lobby_DestroyLobbyOptions DestroyLobbyOptions = {};
+	DestroyLobbyOptions.ApiVersion = EOS_LOBBY_DESTROYLOBBY_API_LATEST;
+	DestroyLobbyOptions.LocalUserId = GetProductUserIdChecked(Params.LocalUserId);
+	DestroyLobbyOptions.LobbyId = Params.LobbyData->GetLobbyIdEOS();
+
+	TPromise<TDefaultErrorResult<FDestroyLobbyImpl>> Promise;
+	TFuture<TDefaultErrorResult<FDestroyLobbyImpl>> Future = Promise.GetFuture();
+
+	EOS_Async<EOS_Lobby_DestroyLobbyCallbackInfo>(EOS_Lobby_DestroyLobby, LobbyPrerequisites->LobbyInterfaceHandle, DestroyLobbyOptions)
+	.Then([Promise = MoveTemp(Promise)](TFuture<const EOS_Lobby_DestroyLobbyCallbackInfo*>&& Future) mutable
+	{
+		const EOS_Lobby_DestroyLobbyCallbackInfo* Result = Future.Get();
+		if (Result->ResultCode != EOS_EResult::EOS_Success)
+		{
+			// Todo: Errors
+			Promise.EmplaceValue(Errors::Unknown(FromEOSError(Result->ResultCode)));
+			return;
+		}
+
+		Promise.EmplaceValue(FDestroyLobbyImpl::Result{});
+	});
+
+	return Future;
+}
+
+TFuture<TDefaultErrorResult<FLobbiesEOSGS::FInviteLobbyMemberImpl>> FLobbiesEOSGS::InviteLobbyMemberImpl(FInviteLobbyMemberImpl::Params&& Params)
+{
+	// Check prerequisites.
+	if (!Services.Get<FAuthEOSGS>()->IsLoggedIn(Params.LocalUserId))
+	{
+		return MakeFulfilledPromise<TDefaultErrorResult<FInviteLobbyMemberImpl>>(Errors::InvalidUser()).GetFuture();
+	}
+
+	if (!Params.LobbyData.IsValid())
+	{
+		return MakeFulfilledPromise<TDefaultErrorResult<FInviteLobbyMemberImpl>>(Errors::InvalidParams()).GetFuture();
+	}
+
+	EOS_Lobby_SendInviteOptions SendInviteOptions = {};
+	SendInviteOptions.ApiVersion = EOS_LOBBY_SENDINVITE_API_LATEST;
+	SendInviteOptions.LocalUserId = GetProductUserIdChecked(Params.LocalUserId);
+	SendInviteOptions.TargetUserId = GetProductUserIdChecked(Params.TargetUserId);
+	SendInviteOptions.LobbyId = Params.LobbyData->GetLobbyIdEOS();
+
+	TPromise<TDefaultErrorResult<FInviteLobbyMemberImpl>> Promise;
+	TFuture<TDefaultErrorResult<FInviteLobbyMemberImpl>> Future = Promise.GetFuture();
+
+	EOS_Async<EOS_Lobby_SendInviteCallbackInfo>(EOS_Lobby_SendInvite, LobbyPrerequisites->LobbyInterfaceHandle, SendInviteOptions)
+	.Then([Promise = MoveTemp(Promise)](TFuture<const EOS_Lobby_SendInviteCallbackInfo*>&& Future) mutable
+	{
+		const EOS_Lobby_SendInviteCallbackInfo* Result = Future.Get();
+		if (Result->ResultCode != EOS_EResult::EOS_Success)
+		{
+			// Todo: Errors
+			Promise.EmplaceValue(Errors::Unknown(FromEOSError(Result->ResultCode)));
+			return;
+		}
+
+		Promise.EmplaceValue(FInviteLobbyMemberImpl::Result{});
+	});
+
+	return Future;
+}
+
+TFuture<TDefaultErrorResult<FLobbiesEOSGS::FDeclineLobbyInvitationImpl>> FLobbiesEOSGS::DeclineLobbyInvitationImpl(FDeclineLobbyInvitationImpl::Params&& Params)
+{
+	// Check prerequisites.
+	if (!Services.Get<FAuthEOSGS>()->IsLoggedIn(Params.LocalUserId))
+	{
+		return MakeFulfilledPromise<TDefaultErrorResult<FDeclineLobbyInvitationImpl>>(Errors::InvalidUser()).GetFuture();
+	}
+
+	// Find the active invitation.
+	TSharedPtr<FLobbyInviteDataEOS> InviteData = GetActiveInvite(Params.LocalUserId, Params.LobbyId);
+	if (!InviteData)
+	{
+		// Todo: Errors.
+		return MakeFulfilledPromise<TDefaultErrorResult<FDeclineLobbyInvitationImpl>>(Errors::Unknown()).GetFuture();
+	}
+
+	EOS_Lobby_RejectInviteOptions RejectInviteOptions = {};
+	RejectInviteOptions.ApiVersion = EOS_LOBBY_REJECTINVITE_API_LATEST;
+	RejectInviteOptions.LocalUserId = GetProductUserIdChecked(Params.LocalUserId);
+	RejectInviteOptions.InviteId = InviteData->GetInviteIdEOS();
+
+	TPromise<TDefaultErrorResult<FDeclineLobbyInvitationImpl>> Promise;
+	TFuture<TDefaultErrorResult<FDeclineLobbyInvitationImpl>> Future = Promise.GetFuture();
+
+	EOS_Async<EOS_Lobby_RejectInviteCallbackInfo>(EOS_Lobby_RejectInvite, LobbyPrerequisites->LobbyInterfaceHandle, RejectInviteOptions)
+	.Then([this, Promise = MoveTemp(Promise), InviteData](TFuture<const EOS_Lobby_RejectInviteCallbackInfo*>&& Future) mutable
+	{
+		// Remove active invitation.
+		RemoveActiveInvite(InviteData.ToSharedRef());
+
+		const EOS_Lobby_RejectInviteCallbackInfo* Result = Future.Get();
+		if (Result->ResultCode != EOS_EResult::EOS_Success)
+		{
+			// Todo: Errors
+			Promise.EmplaceValue(Errors::Unknown(FromEOSError(Result->ResultCode)));
+			return;
+		}
+
+		Promise.EmplaceValue(FDeclineLobbyInvitationImpl::Result{});
+	});
+
+	return Future;
+}
+
+TFuture<TDefaultErrorResult<FLobbiesEOSGS::FKickLobbyMemberImpl>> FLobbiesEOSGS::KickLobbyMemberImpl(FKickLobbyMemberImpl::Params&& Params)
+{
+	// Check prerequisites.
+	if (!Services.Get<FAuthEOSGS>()->IsLoggedIn(Params.LocalUserId))
+	{
+		return MakeFulfilledPromise<TDefaultErrorResult<FKickLobbyMemberImpl>>(Errors::InvalidUser()).GetFuture();
+	}
+
+	if (!Params.LobbyData.IsValid())
+	{
+		return MakeFulfilledPromise<TDefaultErrorResult<FKickLobbyMemberImpl>>(Errors::InvalidParams()).GetFuture();
+	}
+
+	// todo: check local user is lobby owner
+
+	EOS_Lobby_KickMemberOptions KickMemberOptions = {};
+	KickMemberOptions.ApiVersion = EOS_LOBBY_KICKMEMBER_API_LATEST;
+	KickMemberOptions.LocalUserId = GetProductUserIdChecked(Params.LocalUserId);
+	KickMemberOptions.TargetUserId = GetProductUserIdChecked(Params.TargetUserId);
+	KickMemberOptions.LobbyId = Params.LobbyData->GetLobbyIdEOS();
+
+	TPromise<TDefaultErrorResult<FKickLobbyMemberImpl>> Promise;
+	TFuture<TDefaultErrorResult<FKickLobbyMemberImpl>> Future = Promise.GetFuture();
+
+	EOS_Async<EOS_Lobby_KickMemberCallbackInfo>(EOS_Lobby_KickMember, LobbyPrerequisites->LobbyInterfaceHandle, KickMemberOptions)
+	.Then([Promise = MoveTemp(Promise)](TFuture<const EOS_Lobby_KickMemberCallbackInfo*>&& Future) mutable
+	{
+		const EOS_Lobby_KickMemberCallbackInfo* Result = Future.Get();
+		if (Result->ResultCode != EOS_EResult::EOS_Success)
+		{
+			// Todo: Errors
+			Promise.EmplaceValue(Errors::Unknown(FromEOSError(Result->ResultCode)));
+			return;
+		}
+
+		Promise.EmplaceValue(FKickLobbyMemberImpl::Result{});
+	});
+
+	return Future;
+}
+
+TFuture<TDefaultErrorResult<FLobbiesEOSGS::FPromoteLobbyMemberImpl>> FLobbiesEOSGS::PromoteLobbyMemberImpl(FPromoteLobbyMemberImpl::Params&& Params)
+{
+	// Check prerequisites.
+	if (!Services.Get<FAuthEOSGS>()->IsLoggedIn(Params.LocalUserId))
+	{
+		return MakeFulfilledPromise<TDefaultErrorResult<FPromoteLobbyMemberImpl>>(Errors::InvalidUser()).GetFuture();
+	}
+
+	if (!Params.LobbyData.IsValid())
+	{
+		return MakeFulfilledPromise<TDefaultErrorResult<FPromoteLobbyMemberImpl>>(Errors::InvalidParams()).GetFuture();
+	}
+
+	// todo: check local user is lobby owner
+
+	EOS_Lobby_PromoteMemberOptions PromoteMemberOptions = {};
+	PromoteMemberOptions.ApiVersion = EOS_LOBBY_PROMOTEMEMBER_API_LATEST;
+	PromoteMemberOptions.LocalUserId = GetProductUserIdChecked(Params.LocalUserId);
+	PromoteMemberOptions.TargetUserId = GetProductUserIdChecked(Params.TargetUserId);
+	PromoteMemberOptions.LobbyId = Params.LobbyData->GetLobbyIdEOS();
+
+	TPromise<TDefaultErrorResult<FPromoteLobbyMemberImpl>> Promise;
+	TFuture<TDefaultErrorResult<FPromoteLobbyMemberImpl>> Future = Promise.GetFuture();
+
+	EOS_Async<EOS_Lobby_PromoteMemberCallbackInfo>(EOS_Lobby_PromoteMember, LobbyPrerequisites->LobbyInterfaceHandle, PromoteMemberOptions)
+	.Then([Promise = MoveTemp(Promise)](TFuture<const EOS_Lobby_PromoteMemberCallbackInfo*>&& Future) mutable
+	{
+		const EOS_Lobby_PromoteMemberCallbackInfo* Result = Future.Get();
+		if (Result->ResultCode != EOS_EResult::EOS_Success)
+		{
+			// Todo: Errors
+			Promise.EmplaceValue(Errors::Unknown(FromEOSError(Result->ResultCode)));
+			return;
+		}
+
+		Promise.EmplaceValue(FPromoteLobbyMemberImpl::Result{});
+	});
+
+	return Future;
+}
+
+TFuture<TDefaultErrorResult<FLobbiesEOSGS::FModifyLobbyDataImpl>> FLobbiesEOSGS::ModifyLobbyDataImpl(FModifyLobbyDataImpl::Params&& Params)
+{
+	// Check prerequisites.
+	if (!Services.Get<FAuthEOSGS>()->IsLoggedIn(Params.LocalUserId))
+	{
+		return MakeFulfilledPromise<TDefaultErrorResult<FModifyLobbyDataImpl>>(Errors::InvalidUser()).GetFuture();
+	}
+
+	if (!Params.LobbyData.IsValid())
+	{
+		return MakeFulfilledPromise<TDefaultErrorResult<FModifyLobbyDataImpl>>(Errors::InvalidParams()).GetFuture();
+	}
+
+	TSharedPtr<FLobbyDetailsEOS> LobbyDetails = Params.LobbyData->GetUserLobbyDetails(Params.LocalUserId);
+	if (!LobbyDetails)
+	{
+		return MakeFulfilledPromise<TDefaultErrorResult<FModifyLobbyDataImpl>>(Errors::InvalidParams()).GetFuture();
+	}
+
+	if (Params.LocalUserId != Params.LobbyData->GetClientLobbyData()->GetPublicData().OwnerAccountId)
+	{
+		return MakeFulfilledPromise<TDefaultErrorResult<FModifyLobbyDataImpl>>(Errors::InvalidParams()).GetFuture();
+	}
+
+	UE_LOG(LogTemp, Warning, TEXT("[FLobbiesEOSGS::ModifyLobbyDataImpl] Start. Lobby: %s, Member: %s"),
+		*Params.LobbyData->GetLobbyId(), *ToLogString(Params.LocalUserId));
+
+	TPromise<TDefaultErrorResult<FModifyLobbyDataImpl>> Promise;
+	TFuture<TDefaultErrorResult<FModifyLobbyDataImpl>> Future = Promise.GetFuture();
+
+	LobbyDetails->ApplyLobbyDataUpdateFromLocalChanges(Params.LocalUserId, *Params.Changes)
+	.Then([Promise = MoveTemp(Promise)](TFuture<EOS_EResult>&& Future) mutable
+	{
+		UE_LOG(LogTemp, Warning, TEXT("[FLobbiesEOSGS::ModifyLobbyDataImpl] Complete. Result: %s"), *LexToString(Future.Get()));
+
+		// Todo: Handle "no change" better.
+		if (Future.Get() != EOS_EResult::EOS_Success && Future.Get() != EOS_EResult::EOS_NoChange)
+		{
+			// Todo: Errors
+			Promise.EmplaceValue(Errors::Unknown(FromEOSError(Future.Get())));
+			return;
+		}
+
+		Promise.EmplaceValue(FModifyLobbyDataImpl::Result{});
+	});
+
+	return Future;
+}
+
+TFuture<TDefaultErrorResult<FLobbiesEOSGS::FModifyLobbyMemberDataImpl>> FLobbiesEOSGS::ModifyLobbyMemberDataImpl(FModifyLobbyMemberDataImpl::Params&& Params)
+{
+	// Check prerequisites.
+	if (!Services.Get<FAuthEOSGS>()->IsLoggedIn(Params.LocalUserId))
+	{
+		return MakeFulfilledPromise<TDefaultErrorResult<FModifyLobbyMemberDataImpl>>(Errors::InvalidUser()).GetFuture();
+	}
+
+	if (!Params.LobbyData.IsValid())
+	{
+		return MakeFulfilledPromise<TDefaultErrorResult<FModifyLobbyMemberDataImpl>>(Errors::InvalidParams()).GetFuture();
+	}
+
+	TSharedPtr<FLobbyDetailsEOS> LobbyDetails = Params.LobbyData->GetUserLobbyDetails(Params.LocalUserId);
+	if (!LobbyDetails)
+	{
+		return MakeFulfilledPromise<TDefaultErrorResult<FModifyLobbyMemberDataImpl>>(Errors::InvalidParams()).GetFuture();
+	}
+
+	TPromise<TDefaultErrorResult<FModifyLobbyMemberDataImpl>> Promise;
+	TFuture<TDefaultErrorResult<FModifyLobbyMemberDataImpl>> Future = Promise.GetFuture();
+
+	LobbyDetails->ApplyLobbyMemberDataUpdateFromLocalChanges(Params.LocalUserId, *Params.Changes)
+	.Then([Promise = MoveTemp(Promise)](TFuture<EOS_EResult>&& Future) mutable
+	{
+		// Todo: Handle "no change" better.
+		if (Future.Get() != EOS_EResult::EOS_Success && Future.Get() != EOS_EResult::EOS_NoChange)
+		{
+			// Todo: Errors
+			Promise.EmplaceValue(Errors::Unknown(FromEOSError(Future.Get())));
+			return;
+		}
+
+		Promise.EmplaceValue(FModifyLobbyMemberDataImpl::Result{});
+	});
+
+	return Future;
+}
+
+TOnlineAsyncOpHandle<FLobbiesEOSGS::FProcessLobbyNotificationImpl> FLobbiesEOSGS::ProcessLobbyNotificationImplOp(FProcessLobbyNotificationImpl::Params&& InParams)
+{
+	TOnlineAsyncOpRef<FProcessLobbyNotificationImpl> Op = GetOp<FProcessLobbyNotificationImpl>(MoveTemp(InParams));
+	const FProcessLobbyNotificationImpl::Params& Params = Op->GetParams();
+
+	if (!Params.LobbyData.IsValid())
+	{
+		Op->SetError(Errors::InvalidParams());
+		return Op->GetHandle();
+	}
+
+	Op->Then([this](TOnlineAsyncOp<FProcessLobbyNotificationImpl>& InAsyncOp)
+	{
+		const FProcessLobbyNotificationImpl::Params& Params = InAsyncOp.GetParams();
+
+		// Notifications do not always indicate a user. Try to find a valid lobby details object to
+		// handle acquiring data snapshots.
+		TSharedPtr<FLobbyDetailsEOS> LobbyDetails = Params.LobbyData->GetActiveLobbyDetails();
+		if (!LobbyDetails.IsValid())
+		{
+			UE_LOG(LogTemp, Log, TEXT("[FLobbiesEOSGS::ProcessLobbyNotificationImplOp] Failed to find active lobby details to process lobby notificaions: Lobby: %s"),
+				*Params.LobbyData->GetLobbyId());
+			InAsyncOp.SetError(Errors::Unknown());
+			return MakeFulfilledPromise<TDefaultErrorResultInternal<TSharedRef<FClientLobbySnapshot>>>().GetFuture();
+		}
+
+		InAsyncOp.Data.Set<TSharedRef<FLobbyDetailsEOS>>(LobbyDetailsKeyName, LobbyDetails.ToSharedRef());
+
+		// Fetch lobby snapshot. Fetching the snapshot resolves the account ids of all lobby members in the snapshot.
+		return LobbyDetails->GetLobbySnapshot();
+	})
+	.Then([this](TOnlineAsyncOp<FProcessLobbyNotificationImpl>& InAsyncOp, TDefaultErrorResultInternal<TSharedRef<FClientLobbySnapshot>>&& LobbySnapshotResult)
+	{
+		const FProcessLobbyNotificationImpl::Params& Params = InAsyncOp.GetParams();
+		TSharedRef<FLobbyDetailsEOS> LobbyDetails = GetOpDataChecked<TSharedRef<FLobbyDetailsEOS>>(InAsyncOp, LobbyDetailsKeyName);
+
+		if (LobbySnapshotResult.IsError())
+		{
+			// Todo: errors.
+			InAsyncOp.SetError(Errors::Unknown(MoveTempIfPossible(LobbySnapshotResult.GetErrorValue())));
+			return;
+		}
+
+		// Get member snapshots.
+		TMap<FOnlineAccountIdHandle, TSharedRef<FClientLobbyMemberSnapshot>> LobbyMemberSnapshots;
+		LobbyMemberSnapshots.Reserve(Params.MutatedMembers.Num());
+		for (EOS_ProductUserId MutatedMember : Params.MutatedMembers)
+		{
+			const FOnlineAccountIdHandle MutatedMemberAccountId = FindAccountId(MutatedMember);
+			if (MutatedMemberAccountId.IsValid())
+			{
+				TDefaultErrorResultInternal<TSharedRef<FClientLobbyMemberSnapshot>> MemberSnapshotResult = LobbyDetails->GetLobbyMemberSnapshot(MutatedMemberAccountId);
+				if (MemberSnapshotResult.IsError())
+				{
+					// Todo: errors.
+					InAsyncOp.SetError(Errors::Unknown(MoveTempIfPossible(MemberSnapshotResult.GetErrorValue())));
+					return;
+				}
+				LobbyMemberSnapshots.Add(MutatedMemberAccountId, MoveTemp(MemberSnapshotResult.GetOkValue()));
+			}
+		}
+
+		// Translate leaving members from EOS_ProductUserId to FOnlineAccountIdHandle.
+		TMap<FOnlineAccountIdHandle, ELobbyMemberLeaveReason> LeavingMemberReason;
+		LeavingMemberReason.Reserve(Params.LeavingMembers.Num());
+		for (const TPair<EOS_ProductUserId, ELobbyMemberLeaveReason>& LeavingMember : Params.LeavingMembers)
+		{
+			const FOnlineAccountIdHandle LeavingMemberAccountId = FindAccountId(LeavingMember.Key);
+			if (LeavingMemberAccountId.IsValid())
+			{
+				LeavingMemberReason.Add(LeavingMemberAccountId, LeavingMember.Value);
+			}
+		}
+
+		// Apply updates and fire notifications.
+		FApplyLobbyUpdateResult Result = Params.LobbyData->GetClientLobbyData()->ApplyLobbyUpdateFromServiceSnapshot(
+			MoveTemp(*LobbySnapshotResult.GetOkValue()),
+			MoveTemp(LobbyMemberSnapshots),
+			MoveTemp(LeavingMemberReason),
+			&LobbyEvents);
+
+		// Remove active users if needed.
+		for (FOnlineAccountIdHandle LeavingMember : Result.LeavingLocalMembers)
+		{
+			RemoveActiveLobby(LeavingMember, Params.LobbyData.ToSharedRef());
+		}
+
+		InAsyncOp.SetResult(FProcessLobbyNotificationImpl::Result{});
+	})
+	.Enqueue(GetSerialQueue());
+
+	return Op->GetHandle();
+}
+
+/* UE::Online */ }
diff --git a/Engine/Plugins/Online/Experimental/OnlineServicesEOSGS/Source/Private/Online/LobbiesEOSGSTypes.cpp b/Engine/Plugins/Online/Experimental/OnlineServicesEOSGS/Source/Private/Online/LobbiesEOSGSTypes.cpp
new file mode 100644
--- /dev/null
+++ b/Engine/Plugins/Online/Experimental/OnlineServicesEOSGS/Source/Private/Online/LobbiesEOSGSTypes.cpp
@@ -0,0 +1,1037 @@
+// Copyright Epic Games, Inc. All Rights Reserved.
+
+#include "LobbiesEOSGSTypes.h"
+
+#include "Online/AuthEOSGS.h"
+#include "Online/LobbiesCommon.h"
+
+namespace UE::Online {
+
+namespace Private {
+
+FString TranslateLobbyId(EOS_LobbyId EOSLobbyId)
+{
+	return FString(UTF8_TO_TCHAR(EOSLobbyId));
+}
+
+FString TranslateLobbyInviteId(const char* InviteId)
+{
+	return FString(UTF8_TO_TCHAR(InviteId));
+}
+
+FLobbyAttributeId TranslateLobbyAttributeId(const char* AttributeId)
+{
+	return FLobbyAttributeId(UTF8_TO_TCHAR(AttributeId));
+}
+
+EOS_EComparisonOp TranslateSearchComparison(ELobbyComparisonOp Op)
+{
+	switch (Op)
+	{
+	default:									checkNoEntry(); // Intentional fallthrough
+	case ELobbyComparisonOp::Equals:			return EOS_EComparisonOp::EOS_CO_EQUAL;
+	case ELobbyComparisonOp::NotEquals:			return EOS_EComparisonOp::EOS_CO_NOTEQUAL;
+	case ELobbyComparisonOp::GreaterThan:		return EOS_EComparisonOp::EOS_CO_GREATERTHAN;
+	case ELobbyComparisonOp::GreaterThanEquals:	return EOS_EComparisonOp::EOS_CO_GREATERTHANOREQUAL;
+	case ELobbyComparisonOp::LessThan:			return EOS_EComparisonOp::EOS_CO_LESSTHAN;
+	case ELobbyComparisonOp::LessThanEquals:	return EOS_EComparisonOp::EOS_CO_LESSTHANOREQUAL;
+	case ELobbyComparisonOp::Near:				return EOS_EComparisonOp::EOS_CO_DISTANCE;
+	case ELobbyComparisonOp::In:				return EOS_EComparisonOp::EOS_CO_ONEOF;
+	case ELobbyComparisonOp::NotIn:				return EOS_EComparisonOp::EOS_CO_NOTANYOF;
+
+	// todo:
+	// EOS_EComparisonOp::EOS_CO_ANYOF
+	// EOS_EComparisonOp::EOS_CO_NOTONEOF
+	// EOS_EComparisonOp::EOS_CO_CONTAINS
+	}
+}
+
+} // Private
+
+const FString FLobbyBucketIdEOS::Separator = TEXT("|");
+
+FLobbyBucketIdEOS::FLobbyBucketIdEOS(FString ProductName, int32 ProductVersion)
+	: ProductName(ProductName.Replace(*Separator, TEXT("_")))
+	, ProductVersion(ProductVersion)
+{
+}
+
+// Attribute translators.
+FLobbyAttributeTranslator<ELobbyTranslationType::ToService>::FLobbyAttributeTranslator(const TPair<FLobbyAttributeId, FLobbyVariant>& FromAttributeData)
+	: FLobbyAttributeTranslator(FromAttributeData.Key, FromAttributeData.Value)
+{
+}
+
+FLobbyAttributeTranslator<ELobbyTranslationType::ToService>::FLobbyAttributeTranslator(FLobbyAttributeId FromAttributeId, const FLobbyVariant& FromAttributeData)
+	: KeyConverterStorage(*FromAttributeId.ToString())
+{
+	AttributeData.ApiVersion = EOS_LOBBY_ATTRIBUTEDATA_API_LATEST;
+	AttributeData.Key = KeyConverterStorage.Get();
+
+	if (FromAttributeData.VariantData.IsType<FString>())
+	{
+		ValueConverterStorage.Emplace(*FromAttributeData.VariantData.Get<FString>());
+		AttributeData.ValueType = EOS_ELobbyAttributeType::EOS_AT_STRING;
+		AttributeData.Value.AsUtf8 = ValueConverterStorage->Get();
+	}
+	else if (FromAttributeData.VariantData.IsType<int64>())
+	{
+		AttributeData.ValueType = EOS_ELobbyAttributeType::EOS_AT_INT64;
+		AttributeData.Value.AsInt64 = FromAttributeData.VariantData.Get<int64>();
+	}
+	else if (FromAttributeData.VariantData.IsType<double>())
+	{
+		AttributeData.ValueType = EOS_ELobbyAttributeType::EOS_AT_DOUBLE;
+		AttributeData.Value.AsDouble = FromAttributeData.VariantData.Get<double>();
+	}
+	else if (FromAttributeData.VariantData.IsType<bool>())
+	{
+		AttributeData.ValueType = EOS_ELobbyAttributeType::EOS_AT_BOOLEAN;
+		AttributeData.Value.AsBool = FromAttributeData.VariantData.Get<bool>();
+	}
+}
+
+FLobbyAttributeTranslator<ELobbyTranslationType::FromService>::FLobbyAttributeTranslator(const EOS_Lobby_AttributeData& FromAttributeData)
+{
+	FLobbyAttributeId AttributeId = Private::TranslateLobbyAttributeId(FromAttributeData.Key);
+	FLobbyVariant VariantData;
+
+	switch (FromAttributeData.ValueType)
+	{
+	case EOS_ELobbyAttributeType::EOS_AT_BOOLEAN:
+		VariantData.Set(FromAttributeData.Value.AsBool != 0);
+		break;
+
+	case EOS_ELobbyAttributeType::EOS_AT_INT64:
+		VariantData.Set(static_cast<int64>(FromAttributeData.Value.AsInt64));
+		break;
+
+	case EOS_ELobbyAttributeType::EOS_AT_DOUBLE:
+		VariantData.Set(FromAttributeData.Value.AsDouble);
+		break;
+
+	case EOS_ELobbyAttributeType::EOS_AT_STRING:
+		VariantData.Set(UTF8_TO_TCHAR(FromAttributeData.Value.AsUtf8));
+		break;
+
+	default:
+		checkNoEntry();
+		break;
+	}
+
+	AttributeData = TPair<FLobbyAttributeId, FLobbyVariant>(MoveTemp(AttributeId), MoveTemp(VariantData));
+}
+
+FLobbyBucketIdTranslator<ELobbyTranslationType::ToService>::FLobbyBucketIdTranslator(const FLobbyBucketIdEOS& BucketId)
+	: BucketConverterStorage(*FString::Printf(TEXT("%s%s%d"), *BucketId.GetProductName(), *FLobbyBucketIdEOS::Separator, BucketId.GetProductVersion()))
+{
+}
+
+FLobbyBucketIdTranslator<ELobbyTranslationType::FromService>::FLobbyBucketIdTranslator(const char* BucketIdEOS)
+{
+	FUTF8ToTCHAR BucketConverterStorage(BucketIdEOS);
+	FString BucketString(BucketConverterStorage.Get());
+
+	constexpr int32 ExpectedPartsNum = 2;
+	TArray<FString> Parts;
+	if (BucketString.ParseIntoArray(Parts, *FLobbyBucketIdEOS::Separator) == ExpectedPartsNum)
+	{
+		int32 BuildId = 0;
+		::LexFromString(BuildId, *Parts[1]);
+		BucketId = FLobbyBucketIdEOS(Parts[0], BuildId);
+	}
+}
+
+const EOS_HLobbyDetails FLobbyDetailsEOS::InvalidLobbyDetailsHandle = {};
+
+TDefaultErrorResultInternal<TSharedRef<FLobbyDetailsEOS>> FLobbyDetailsEOS::CreateFromLobbyId(
+	const TSharedRef<FLobbyPrerequisitesEOS>& Prerequisites,
+	FOnlineAccountIdHandle LocalUserId,
+	EOS_LobbyId LobbyId)
+{
+	EOS_HLobbyDetails LobbyDetailsHandle = {};
+
+	EOS_Lobby_CopyLobbyDetailsHandleOptions Options;
+	Options.ApiVersion = EOS_LOBBY_COPYLOBBYDETAILSHANDLE_API_LATEST;
+	Options.LobbyId = LobbyId;
+	Options.LocalUserId = GetProductUserIdChecked(LocalUserId);
+
+	EOS_EResult EOSResult = EOS_Lobby_CopyLobbyDetailsHandle(Prerequisites->LobbyInterfaceHandle, &Options, &LobbyDetailsHandle);
+	if (EOSResult != EOS_EResult::EOS_Success)
+	{
+		return TDefaultErrorResultInternal<TSharedRef<FLobbyDetailsEOS>>(FromEOSError(EOSResult));
+	}
+
+	TDefaultErrorResultInternal<TSharedRef<FLobbyDetailsInfoEOS>> LobbyDetailsInfoResult = FLobbyDetailsInfoEOS::Create(LobbyDetailsHandle);
+	if (LobbyDetailsInfoResult.IsError())
+	{
+		EOS_LobbyDetails_Release(LobbyDetailsHandle);
+		return TDefaultErrorResultInternal<TSharedRef<FLobbyDetailsEOS>>(MoveTemp(LobbyDetailsInfoResult.GetErrorValue()));
+	}
+
+	return TDefaultErrorResultInternal<TSharedRef<FLobbyDetailsEOS>>(MakeShared<FLobbyDetailsEOS>(Prerequisites, LobbyDetailsInfoResult.GetOkValue(), LocalUserId, ELobbyDetailsSource::Active, LobbyDetailsHandle));
+}
+
+TDefaultErrorResultInternal<TSharedRef<FLobbyDetailsEOS>> FLobbyDetailsEOS::CreateFromInviteId(
+	const TSharedRef<FLobbyPrerequisitesEOS>& Prerequisites,
+	FOnlineAccountIdHandle LocalUserId,
+	const char* InviteId)
+{
+	EOS_HLobbyDetails LobbyDetailsHandle = {};
+
+	EOS_Lobby_CopyLobbyDetailsHandleByInviteIdOptions Options;
+	Options.ApiVersion = EOS_LOBBY_COPYLOBBYDETAILSHANDLEBYINVITEID_API_LATEST;
+	Options.InviteId = InviteId;
+
+	EOS_EResult EOSResult = EOS_Lobby_CopyLobbyDetailsHandleByInviteId(Prerequisites->LobbyInterfaceHandle, &Options, &LobbyDetailsHandle);
+	if (EOSResult != EOS_EResult::EOS_Success)
+	{
+		return TDefaultErrorResultInternal<TSharedRef<FLobbyDetailsEOS>>(FromEOSError(EOSResult));
+	}
+
+	TDefaultErrorResultInternal<TSharedRef<FLobbyDetailsInfoEOS>> LobbyDetailsInfoResult = FLobbyDetailsInfoEOS::Create(LobbyDetailsHandle);
+	if (LobbyDetailsInfoResult.IsError())
+	{
+		EOS_LobbyDetails_Release(LobbyDetailsHandle);
+		return TDefaultErrorResultInternal<TSharedRef<FLobbyDetailsEOS>>(MoveTemp(LobbyDetailsInfoResult.GetErrorValue()));
+	}
+
+	return TDefaultErrorResultInternal<TSharedRef<FLobbyDetailsEOS>>(MakeShared<FLobbyDetailsEOS>(Prerequisites, LobbyDetailsInfoResult.GetOkValue(), LocalUserId, ELobbyDetailsSource::Invite, LobbyDetailsHandle));
+}
+
+TDefaultErrorResultInternal<TSharedRef<FLobbyDetailsEOS>> FLobbyDetailsEOS::CreateFromUiEventId(
+	const TSharedRef<FLobbyPrerequisitesEOS>& Prerequisites,
+	FOnlineAccountIdHandle LocalUserId,
+	EOS_UI_EventId UiEventId)
+{
+	EOS_HLobbyDetails LobbyDetailsHandle = {};
+
+	EOS_Lobby_CopyLobbyDetailsHandleByUiEventIdOptions Options;
+	Options.ApiVersion = EOS_LOBBY_COPYLOBBYDETAILSHANDLEBYUIEVENTID_API_LATEST;
+	Options.UiEventId = UiEventId;
+
+	EOS_EResult EOSResult = EOS_Lobby_CopyLobbyDetailsHandleByUiEventId(Prerequisites->LobbyInterfaceHandle, &Options, &LobbyDetailsHandle);
+	if (EOSResult != EOS_EResult::EOS_Success)
+	{
+		return TDefaultErrorResultInternal<TSharedRef<FLobbyDetailsEOS>>(FromEOSError(EOSResult));
+	}
+
+	TDefaultErrorResultInternal<TSharedRef<FLobbyDetailsInfoEOS>> LobbyDetailsInfoResult = FLobbyDetailsInfoEOS::Create(LobbyDetailsHandle);
+	if (LobbyDetailsInfoResult.IsError())
+	{
+		EOS_LobbyDetails_Release(LobbyDetailsHandle);
+		return TDefaultErrorResultInternal<TSharedRef<FLobbyDetailsEOS>>(MoveTemp(LobbyDetailsInfoResult.GetErrorValue()));
+	}
+
+	return TDefaultErrorResultInternal<TSharedRef<FLobbyDetailsEOS>>(MakeShared<FLobbyDetailsEOS>(Prerequisites, LobbyDetailsInfoResult.GetOkValue(), LocalUserId, ELobbyDetailsSource::UiEvent, LobbyDetailsHandle));
+}
+
+TDefaultErrorResultInternal<TSharedRef<FLobbyDetailsEOS>> FLobbyDetailsEOS::CreateFromSearchResult(
+	const TSharedRef<FLobbyPrerequisitesEOS>& Prerequisites,
+	FOnlineAccountIdHandle LocalUserId,
+	EOS_HLobbySearch SearchHandle,
+	uint32_t ResultIndex)
+{
+	EOS_HLobbyDetails LobbyDetailsHandle = {};
+
+	EOS_LobbySearch_CopySearchResultByIndexOptions Options;
+	Options.ApiVersion = EOS_LOBBYSEARCH_COPYSEARCHRESULTBYINDEX_API_LATEST;
+	Options.LobbyIndex = ResultIndex;
+
+	EOS_EResult EOSResult = EOS_LobbySearch_CopySearchResultByIndex(SearchHandle, &Options, &LobbyDetailsHandle);
+	if (EOSResult != EOS_EResult::EOS_Success)
+	{
+		return TDefaultErrorResultInternal<TSharedRef<FLobbyDetailsEOS>>(FromEOSError(EOSResult));
+	}
+
+	TDefaultErrorResultInternal<TSharedRef<FLobbyDetailsInfoEOS>> LobbyDetailsInfoResult = FLobbyDetailsInfoEOS::Create(LobbyDetailsHandle);
+	if (LobbyDetailsInfoResult.IsError())
+	{
+		EOS_LobbyDetails_Release(LobbyDetailsHandle);
+		return TDefaultErrorResultInternal<TSharedRef<FLobbyDetailsEOS>>(MoveTemp(LobbyDetailsInfoResult.GetErrorValue()));
+	}
+
+	return TDefaultErrorResultInternal<TSharedRef<FLobbyDetailsEOS>>(MakeShared<FLobbyDetailsEOS>(Prerequisites, LobbyDetailsInfoResult.GetOkValue(), LocalUserId, ELobbyDetailsSource::Search, LobbyDetailsHandle));
+}
+
+FLobbyDetailsEOS::~FLobbyDetailsEOS()
+{
+	EOS_LobbyDetails_Release(LobbyDetailsHandle);
+}
+
+TFuture<TDefaultErrorResultInternal<TSharedRef<FClientLobbySnapshot>>> FLobbyDetailsEOS::GetLobbySnapshot() const
+{
+	TSharedPtr<FAuthEOSGS> AuthInterface = Prerequisites->AuthInterface.Pin();
+	if (!AuthInterface)
+	{
+		return MakeFulfilledPromise<TDefaultErrorResultInternal<TSharedRef<FClientLobbySnapshot>>>(Errors::MissingInterface()).GetFuture();
+	}
+
+	EOS_LobbyDetails_GetMemberCountOptions GetMemberCountOptions = {};
+	GetMemberCountOptions.ApiVersion = EOS_LOBBYDETAILS_GETMEMBERCOUNT_API_LATEST;
+	const uint32_t MemberCount = EOS_LobbyDetails_GetMemberCount(LobbyDetailsHandle, &GetMemberCountOptions);
+
+	TSharedRef<TArray<EOS_ProductUserId>> MemberProductUserIds = MakeShared<TArray<EOS_ProductUserId>>();
+	MemberProductUserIds->Reserve(MemberCount);
+
+	for (uint32_t MemberIndex = 0; MemberIndex < MemberCount; ++MemberIndex)
+	{
+		EOS_LobbyDetails_GetMemberByIndexOptions GetMemberByIndexOptions = {};
+		GetMemberByIndexOptions.ApiVersion = EOS_LOBBYDETAILS_GETMEMBERBYINDEX_API_LATEST;
+		GetMemberByIndexOptions.MemberIndex = MemberIndex;
+		MemberProductUserIds->Emplace(EOS_LobbyDetails_GetMemberByIndex(LobbyDetailsHandle, &GetMemberByIndexOptions));
+	}
+
+	TPromise<TDefaultErrorResultInternal<TSharedRef<FClientLobbySnapshot>>> Promise;
+	TFuture<TDefaultErrorResultInternal<TSharedRef<FClientLobbySnapshot>>> Future = Promise.GetFuture();
+
+	// Resolve lobby member product user ids to FOnlineAccountIdHandle before proceeding.
+	AuthInterface->ResolveAccountIds(AssociatedLocalUser, *MemberProductUserIds)
+	.Then([StrongThis = AsShared(), Promise = MoveTemp(Promise), MemberProductUserIds](TFuture<TArray<FOnlineAccountIdHandle>>&& Future) mutable
+	{
+		const TArray<FOnlineAccountIdHandle>& ResolvedAccountIds = Future.Get();
+		if (MemberProductUserIds->Num() != ResolvedAccountIds.Num())
+		{
+			// Todo: Errors
+			Promise.EmplaceValue(Errors::Unknown());
+			return;
+		}
+
+		TSharedRef<FClientLobbySnapshot> ClientLobbySnapshot = MakeShared<FClientLobbySnapshot>();
+		ClientLobbySnapshot->MaxMembers = StrongThis->GetInfo()->GetMaxMembers();
+		ClientLobbySnapshot->JoinPolicy = TranslateJoinPolicy(StrongThis->GetInfo()->GetPermissionLevel());
+
+		// Resolve member info.
+		{
+			EOS_LobbyDetails_GetLobbyOwnerOptions GetLobbyOwnerOptions = {};
+			GetLobbyOwnerOptions.ApiVersion = EOS_LOBBYDETAILS_GETLOBBYOWNER_API_LATEST;
+			const EOS_ProductUserId LobbyOwner = EOS_LobbyDetails_GetLobbyOwner(StrongThis->LobbyDetailsHandle, &GetLobbyOwnerOptions);
+
+			for (int32 MemberIndex = 0; MemberIndex < MemberProductUserIds->Num(); ++MemberIndex)
+			{
+				const EOS_ProductUserId MemberProductUserId = (*MemberProductUserIds)[MemberIndex];
+				const FOnlineAccountIdHandle ResolvedMemberAccountId = ResolvedAccountIds[MemberIndex];
+
+				if (MemberProductUserId == LobbyOwner)
+				{
+					ClientLobbySnapshot->OwnerAccountId = ResolvedMemberAccountId;
+				}
+
+				ClientLobbySnapshot->Members.Add(ResolvedMemberAccountId);
+			}
+		}
+
+		// Resolve lobby attributes
+		{
+			EOS_LobbyDetails_GetAttributeCountOptions GetAttributeCountOptions = {};
+			GetAttributeCountOptions.ApiVersion = EOS_LOBBYDETAILS_GETATTRIBUTECOUNT_API_LATEST;
+
+			const uint32_t AttributeCount = EOS_LobbyDetails_GetAttributeCount(StrongThis->LobbyDetailsHandle, &GetAttributeCountOptions);
+			for (uint32_t AttributeIndex = 0; AttributeIndex < AttributeCount; ++AttributeIndex)
+			{
+				EOS_LobbyDetails_CopyAttributeByIndexOptions CopyAttributeByIndexOptions = {};
+				CopyAttributeByIndexOptions.ApiVersion = EOS_LOBBYDETAILS_COPYATTRIBUTEBYINDEX_API_LATEST;
+				CopyAttributeByIndexOptions.AttrIndex = AttributeIndex;
+
+				EOS_Lobby_Attribute* LobbyAttribute = nullptr;
+				ON_SCOPE_EXIT
+				{
+					EOS_Lobby_Attribute_Release(LobbyAttribute);
+				};
+
+				EOS_EResult EOSResult = EOS_LobbyDetails_CopyAttributeByIndex(StrongThis->LobbyDetailsHandle, &CopyAttributeByIndexOptions, &LobbyAttribute);
+				if (EOSResult != EOS_EResult::EOS_Success)
+				{
+					// todo: errors
+					Promise.EmplaceValue(FromEOSError(EOSResult));
+					return;
+				}
+
+				FLobbyAttributeTranslator<ELobbyTranslationType::FromService> AttributeTranslator(*LobbyAttribute->Data);
+				ClientLobbySnapshot->Attributes.Add(MoveTemp(AttributeTranslator.GetMutableAttributeData()));
+			}
+		}
+
+		Promise.EmplaceValue(MoveTemp(ClientLobbySnapshot));
+	});
+
+	return Future;
+}
+
+TDefaultErrorResultInternal<TSharedRef<FClientLobbyMemberSnapshot>> FLobbyDetailsEOS::GetLobbyMemberSnapshot(FOnlineAccountIdHandle MemberAccountId) const
+{
+	EOS_ProductUserId MemberProductUserId = GetProductUserIdChecked(MemberAccountId);
+
+	TSharedRef<FClientLobbyMemberSnapshot> LobbyMemberSnapshot = MakeShared<FClientLobbyMemberSnapshot>();
+	LobbyMemberSnapshot->AccountId = MemberAccountId;
+	// Todo: 
+	//ClientMemberData->PlatformAccountId;
+	//ClientMemberData->PlatformDisplayName;
+
+	// Fetch attributes.
+	{
+		EOS_LobbyDetails_GetMemberAttributeCountOptions GetMemberAttributeCountOptions = {};
+		GetMemberAttributeCountOptions.ApiVersion = EOS_LOBBYDETAILS_GETMEMBERATTRIBUTECOUNT_API_LATEST;
+		GetMemberAttributeCountOptions.TargetUserId = MemberProductUserId;
+
+		const uint32_t MemberAttributeCount = EOS_LobbyDetails_GetMemberAttributeCount(LobbyDetailsHandle, &GetMemberAttributeCountOptions);
+		for (uint32_t MemberAttributeIndex = 0; MemberAttributeIndex < MemberAttributeCount; ++MemberAttributeIndex)
+		{
+			EOS_LobbyDetails_CopyMemberAttributeByIndexOptions CopyMemberAttributeByIndexOptions = {};
+			CopyMemberAttributeByIndexOptions.ApiVersion = EOS_LOBBYDETAILS_COPYMEMBERATTRIBUTEBYINDEX_API_LATEST;
+			CopyMemberAttributeByIndexOptions.TargetUserId = MemberProductUserId;
+			CopyMemberAttributeByIndexOptions.AttrIndex = MemberAttributeIndex;
+
+			EOS_Lobby_Attribute* LobbyAttribute = nullptr;
+			ON_SCOPE_EXIT
+			{
+				EOS_Lobby_Attribute_Release(LobbyAttribute);
+			};
+
+			EOS_EResult EOSResult = EOS_LobbyDetails_CopyMemberAttributeByIndex(LobbyDetailsHandle, &CopyMemberAttributeByIndexOptions, &LobbyAttribute);
+			if (EOSResult != EOS_EResult::EOS_Success)
+			{
+				return TDefaultErrorResultInternal<TSharedRef<FClientLobbyMemberSnapshot>>(FromEOSError(EOSResult));
+			}
+
+			FLobbyAttributeTranslator<ELobbyTranslationType::FromService> AttributeTranslator(*LobbyAttribute->Data);
+			LobbyMemberSnapshot->Attributes.Add(MoveTemp(AttributeTranslator.GetMutableAttributeData()));
+		}
+	}
+
+	return TDefaultErrorResultInternal<TSharedRef<FClientLobbyMemberSnapshot>>(MoveTemp(LobbyMemberSnapshot));
+}
+
+TFuture<EOS_EResult> FLobbyDetailsEOS::ApplyLobbyDataUpdateFromLocalChanges(
+	FOnlineAccountIdHandle LocalUserId,
+	const FClientLobbyDataChanges& Changes) const
+{
+	EOS_HLobbyModification LobbyModificationHandle = {};
+
+	ON_SCOPE_EXIT
+	{
+		EOS_LobbyModification_Release(LobbyModificationHandle);
+	};
+
+	// Create lobby modification handle.
+	EOS_Lobby_UpdateLobbyModificationOptions ModificationOptions = {};
+	ModificationOptions.ApiVersion = EOS_LOBBY_UPDATELOBBYMODIFICATION_API_LATEST;
+	ModificationOptions.LocalUserId = GetProductUserIdChecked(LocalUserId);
+	ModificationOptions.LobbyId = GetInfo()->GetLobbyId();
+
+	EOS_EResult EOSResultCode = EOS_Lobby_UpdateLobbyModification(Prerequisites->LobbyInterfaceHandle, &ModificationOptions, &LobbyModificationHandle);
+	if (EOSResultCode != EOS_EResult::EOS_Success)
+	{
+		// Todo: Errors
+		return MakeFulfilledPromise<EOS_EResult>(EOSResultCode).GetFuture();
+	}
+
+	if (Changes.JoinPolicy)
+	{
+		// Set lobby join policy.
+		EOS_LobbyModification_SetPermissionLevelOptions SetPermissionOptions = {};
+		SetPermissionOptions.ApiVersion = EOS_LOBBYMODIFICATION_SETPERMISSIONLEVEL_API_LATEST;
+		SetPermissionOptions.PermissionLevel = TranslateJoinPolicy(*Changes.JoinPolicy);
+
+		EOSResultCode = EOS_LobbyModification_SetPermissionLevel(LobbyModificationHandle, &SetPermissionOptions);
+		if (EOSResultCode != EOS_EResult::EOS_Success)
+		{
+			// Todo: Errors
+			return MakeFulfilledPromise<EOS_EResult>(EOSResultCode).GetFuture();
+		}
+	}
+
+	// Add attributes.
+	for (const TPair<FLobbyAttributeId, FLobbyVariant>& MutatedAttribute : Changes.MutatedAttributes)
+	{
+		const FLobbyAttributeTranslator<ELobbyTranslationType::ToService> AttributeTranslator(MutatedAttribute);
+
+		EOS_LobbyModification_AddAttributeOptions AddAttributeOptions = {};
+		AddAttributeOptions.ApiVersion = EOS_LOBBYMODIFICATION_ADDATTRIBUTE_API_LATEST;
+		AddAttributeOptions.Attribute = &AttributeTranslator.GetAttributeData();
+		AddAttributeOptions.Visibility = EOS_ELobbyAttributeVisibility::EOS_LAT_PUBLIC; // todo - get from schema
+		EOSResultCode = EOS_LobbyModification_AddAttribute(LobbyModificationHandle, &AddAttributeOptions);
+		if (EOSResultCode != EOS_EResult::EOS_Success)
+		{
+			// Todo: Errors
+			return MakeFulfilledPromise<EOS_EResult>(EOSResultCode).GetFuture();
+		}
+	}
+
+	// Remove attributes.
+	for (const FLobbyAttributeId& ClearedAttribute : Changes.ClearedAttributes)
+	{
+		const FTCHARToUTF8 KeyConverter(*ClearedAttribute.ToString());
+
+		EOS_LobbyModification_RemoveAttributeOptions RemoveAttributeOptions = {};
+		RemoveAttributeOptions.ApiVersion = EOS_LOBBYMODIFICATION_REMOVEATTRIBUTE_API_LATEST;
+		RemoveAttributeOptions.Key = KeyConverter.Get();
+		EOSResultCode = EOS_LobbyModification_RemoveAttribute(LobbyModificationHandle, &RemoveAttributeOptions);
+		if (EOSResultCode != EOS_EResult::EOS_Success)
+		{
+			// Todo: Errors
+			return MakeFulfilledPromise<EOS_EResult>(EOSResultCode).GetFuture();
+		}
+	}
+
+	TPromise<EOS_EResult> Promise;
+	TFuture<EOS_EResult> Future = Promise.GetFuture();
+
+	// Apply lobby updates.
+	EOS_Lobby_UpdateLobbyOptions UpdateLobbyOptions = {};
+	UpdateLobbyOptions.ApiVersion = EOS_LOBBY_UPDATELOBBY_API_LATEST;
+	UpdateLobbyOptions.LobbyModificationHandle = LobbyModificationHandle;
+
+	EOS_Async<EOS_Lobby_UpdateLobbyCallbackInfo>(EOS_Lobby_UpdateLobby, Prerequisites->LobbyInterfaceHandle, UpdateLobbyOptions)
+	.Then([Promise = MoveTemp(Promise)](TFuture<const EOS_Lobby_UpdateLobbyCallbackInfo*> Future) mutable
+	{
+		Promise.EmplaceValue(Future.Get()->ResultCode);
+	});
+
+	return Future;
+}
+
+TFuture<EOS_EResult> FLobbyDetailsEOS::ApplyLobbyMemberDataUpdateFromLocalChanges(
+	FOnlineAccountIdHandle LocalUserId,
+	const FClientLobbyMemberDataChanges& Changes) const
+{
+	EOS_HLobbyModification LobbyModificationHandle = {};
+
+	ON_SCOPE_EXIT
+	{
+		EOS_LobbyModification_Release(LobbyModificationHandle);
+	};
+
+	// Create lobby modification handle.
+	EOS_Lobby_UpdateLobbyModificationOptions ModificationOptions = {};
+	ModificationOptions.ApiVersion = EOS_LOBBY_UPDATELOBBYMODIFICATION_API_LATEST;
+	ModificationOptions.LocalUserId = GetProductUserIdChecked(LocalUserId);
+	ModificationOptions.LobbyId = GetInfo()->GetLobbyId();
+
+	EOS_EResult EOSResultCode = EOS_Lobby_UpdateLobbyModification(Prerequisites->LobbyInterfaceHandle, &ModificationOptions, &LobbyModificationHandle);
+	if (EOSResultCode != EOS_EResult::EOS_Success)
+	{
+		// Todo: Errors
+		return MakeFulfilledPromise<EOS_EResult>(EOSResultCode).GetFuture();
+	}
+
+	// Add member attributes.
+	for (const TPair<FLobbyAttributeId, FLobbyVariant>& MutatedAttribute : Changes.MutatedAttributes)
+	{
+		const FLobbyAttributeTranslator<ELobbyTranslationType::ToService> AttributeTranslator(MutatedAttribute);
+
+		EOS_LobbyModification_AddMemberAttributeOptions AddMemberAttributeOptions = {};
+		AddMemberAttributeOptions.ApiVersion = EOS_LOBBYMODIFICATION_ADDMEMBERATTRIBUTE_API_LATEST;
+		AddMemberAttributeOptions.Attribute = &AttributeTranslator.GetAttributeData();
+		AddMemberAttributeOptions.Visibility = EOS_ELobbyAttributeVisibility::EOS_LAT_PUBLIC; // todo - get from schema
+		EOSResultCode = EOS_LobbyModification_AddMemberAttribute(LobbyModificationHandle, &AddMemberAttributeOptions);
+		if (EOSResultCode != EOS_EResult::EOS_Success)
+		{
+			// Todo: Errors
+			return MakeFulfilledPromise<EOS_EResult>(EOSResultCode).GetFuture();
+		}
+	}
+
+	// Remove member attributes.
+	for (const FLobbyAttributeId& ClearedAttribute : Changes.ClearedAttributes)
+	{
+		const FTCHARToUTF8 KeyConverter(*ClearedAttribute.ToString());
+
+		EOS_LobbyModification_RemoveMemberAttributeOptions RemoveMemberAttributeOptions = {};
+		RemoveMemberAttributeOptions.ApiVersion = EOS_LOBBYMODIFICATION_REMOVEMEMBERATTRIBUTE_API_LATEST;
+		RemoveMemberAttributeOptions.Key = KeyConverter.Get();
+		EOSResultCode = EOS_LobbyModification_RemoveMemberAttribute(LobbyModificationHandle, &RemoveMemberAttributeOptions);
+		if (EOSResultCode != EOS_EResult::EOS_Success)
+		{
+			// Todo: Errors
+			return MakeFulfilledPromise<EOS_EResult>(EOSResultCode).GetFuture();
+		}
+	}
+
+	// Apply lobby updates.
+	EOS_Lobby_UpdateLobbyOptions UpdateLobbyOptions = {};
+	UpdateLobbyOptions.ApiVersion = EOS_LOBBY_UPDATELOBBY_API_LATEST;
+	UpdateLobbyOptions.LobbyModificationHandle = LobbyModificationHandle;
+
+	TPromise<EOS_EResult> Promise;
+	TFuture<EOS_EResult> Future = Promise.GetFuture();
+
+	EOS_Async<EOS_Lobby_UpdateLobbyCallbackInfo>(EOS_Lobby_UpdateLobby, Prerequisites->LobbyInterfaceHandle, UpdateLobbyOptions)
+	.Then([Promise = MoveTemp(Promise)](TFuture<const EOS_Lobby_UpdateLobbyCallbackInfo*> Future) mutable
+	{
+		Promise.EmplaceValue(Future.Get()->ResultCode);
+	});
+
+	return Future;
+}
+
+FLobbyDetailsEOS::FLobbyDetailsEOS(
+	const TSharedRef<FLobbyPrerequisitesEOS>& Prerequisites,
+	const TSharedRef<FLobbyDetailsInfoEOS>& LobbyDetailsInfo,
+	FOnlineAccountIdHandle LocalUserId,
+	ELobbyDetailsSource LobbyDetailsSource,
+	EOS_HLobbyDetails LobbyDetailsHandle)
+	: Prerequisites(Prerequisites)
+	, LobbyDetailsInfo(LobbyDetailsInfo)
+	, AssociatedLocalUser(LocalUserId)
+	, LobbyDetailsSource(LobbyDetailsSource)
+	, LobbyDetailsHandle(LobbyDetailsHandle)
+{
+}
+
+TDefaultErrorResultInternal<TSharedRef<FLobbyDetailsInfoEOS>> FLobbyDetailsInfoEOS::Create(EOS_HLobbyDetails LobbyDetailsHandle)
+{
+	EOS_LobbyDetails_CopyInfoOptions CopyInfoOptions = {};
+	CopyInfoOptions.ApiVersion = EOS_LOBBYDETAILS_COPYINFO_API_LATEST;
+
+	EOS_LobbyDetails_Info* LobbyDetailsInfo = nullptr;
+	EOS_EResult EOSResult = EOS_LobbyDetails_CopyInfo(LobbyDetailsHandle, &CopyInfoOptions, &LobbyDetailsInfo);
+	if (EOSResult != EOS_EResult::EOS_Success)
+	{
+		return TDefaultErrorResultInternal<TSharedRef<FLobbyDetailsInfoEOS>>(FromEOSError(EOSResult));
+	}
+
+	return TDefaultErrorResultInternal<TSharedRef<FLobbyDetailsInfoEOS>>(MakeShared<FLobbyDetailsInfoEOS>(FLobbyDetailsInfoPtr(LobbyDetailsInfo)));
+}
+
+FLobbyDetailsInfoEOS::FLobbyDetailsInfoEOS(FLobbyDetailsInfoPtr&& InLobbyDetailsInfo)
+	: LobbyDetailsInfo(MoveTempIfPossible(InLobbyDetailsInfo))
+{
+	const FLobbyBucketIdTranslator<ELobbyTranslationType::FromService> BucketTranslator(LobbyDetailsInfo->BucketId);
+	BucketId = BucketTranslator.GetBucketId();
+
+	if (!BucketId.IsValid())
+	{
+		UE_LOG(LogTemp, Warning, TEXT("[FLobbyDetailsInfoEOS] Failed to parse lobby bucket id. Lobby: %s, Bucket: %s"), UTF8_TO_TCHAR(LobbyDetailsInfo->LobbyId), UTF8_TO_TCHAR(LobbyDetailsInfo->BucketId));
+	}
+}
+
+FLobbyDataEOS::~FLobbyDataEOS()
+{
+	if (UnregisterFn)
+	{
+		UnregisterFn(ClientLobbyData->GetPublicData().LobbyId);
+	}
+}
+
+void FLobbyDataEOS::AddUserLobbyDetails(FOnlineAccountIdHandle LocalUserId, const TSharedPtr<FLobbyDetailsEOS>& LobbyDetails)
+{
+	if (TSharedPtr<FLobbyDetailsEOS> ExistingDetails = GetUserLobbyDetails(LocalUserId))
+	{
+		if (ExistingDetails->GetDetailsSource() < LobbyDetails->GetDetailsSource())
+		{
+			return;
+		}
+	}
+
+	UserLobbyDetails.Add(LocalUserId, LobbyDetails);
+}
+
+TSharedPtr<FLobbyDetailsEOS> FLobbyDataEOS::GetUserLobbyDetails(FOnlineAccountIdHandle LocalUserId) const
+{
+	const TSharedPtr<FLobbyDetailsEOS>* Result = UserLobbyDetails.Find(LocalUserId);
+	return Result ? *Result : TSharedPtr<FLobbyDetailsEOS>();
+}
+
+TSharedPtr<FLobbyDetailsEOS> FLobbyDataEOS::GetActiveLobbyDetails() const
+{
+	TSharedPtr<FLobbyDetailsEOS> FoundDetails;
+
+	for (const TPair<FOnlineAccountIdHandle, TSharedPtr<FLobbyDetailsEOS>>& LobbyDetails : UserLobbyDetails)
+	{
+		if (LobbyDetails.Value->GetDetailsSource() == ELobbyDetailsSource::Active)
+		{
+			FoundDetails = LobbyDetails.Value;
+			break;
+		}
+	}
+
+	return FoundDetails;
+}
+
+FLobbyDataEOS::FLobbyDataEOS(
+	const TSharedRef<FClientLobbyData>& ClientLobbyData,
+	const TSharedRef<FLobbyDetailsInfoEOS>& LobbyDetailsInfo,
+	FUnregisterFn UnregisterFn)
+	: ClientLobbyData(ClientLobbyData)
+	, LobbyDetailsInfo(LobbyDetailsInfo)
+	, UnregisterFn(MoveTemp(UnregisterFn))
+	, LobbyId(Private::TranslateLobbyId(LobbyDetailsInfo->GetLobbyId()))
+{
+}
+
+TFuture<TDefaultErrorResultInternal<TSharedRef<FLobbyDataEOS>>> FLobbyDataEOS::Create(
+	FOnlineLobbyIdHandle LobbyIdHandle,
+	const TSharedRef<FLobbyDetailsEOS>& LobbyDetails,
+	FUnregisterFn UnregisterFn)
+{
+	TPromise<TDefaultErrorResultInternal<TSharedRef<FLobbyDataEOS>>> Promise;
+	TFuture<TDefaultErrorResultInternal<TSharedRef<FLobbyDataEOS>>> Future = Promise.GetFuture();
+	
+	LobbyDetails->GetLobbySnapshot()
+	.Then(
+	[
+		Promise = MoveTemp(Promise),
+		LobbyIdHandle,
+		LobbyDetails,
+		UnregisterFn = MoveTemp(UnregisterFn)
+	]
+	(TFuture<TDefaultErrorResultInternal<TSharedRef<FClientLobbySnapshot>>>&& Future) mutable
+	{
+		if (Future.Get().IsError())
+		{
+			// todo: errors.
+			Promise.EmplaceValue(MoveTemp(Future.Get().GetErrorValue()));
+			return;
+		}
+
+		TSharedRef<FClientLobbySnapshot> LobbySnapshot = Future.Get().GetOkValue();
+		TSharedRef<FClientLobbyData> LobbyData = MakeShared<FClientLobbyData>(LobbyIdHandle);
+
+		// Fetch member data and apply them to the lobby.
+		TMap<FOnlineAccountIdHandle, TSharedRef<FClientLobbyMemberSnapshot>> MemberSnapshots;
+		for (FOnlineAccountIdHandle MemberAccountId : LobbySnapshot->Members)
+		{
+			TDefaultErrorResultInternal<TSharedRef<FClientLobbyMemberSnapshot>> LobbyMemberSnapshotResult = LobbyDetails->GetLobbyMemberSnapshot(MemberAccountId);
+			if (LobbyMemberSnapshotResult.IsError())
+			{
+				// todo: errors.
+				Promise.EmplaceValue(MoveTemp(LobbyMemberSnapshotResult.GetErrorValue()));
+				return;
+			}
+
+			MemberSnapshots.Emplace(MemberAccountId, MoveTemp(LobbyMemberSnapshotResult.GetOkValue()));
+		}
+
+		LobbyData->ApplyLobbyUpdateFromServiceSnapshot(MoveTemp(*LobbySnapshot), MoveTemp(MemberSnapshots));
+
+		Promise.EmplaceValue(MakeShared<FLobbyDataEOS>(
+			LobbyData,
+			LobbyDetails->GetInfo(),
+			MoveTemp(UnregisterFn)));
+	});
+
+	return Future;
+}
+
+FLobbyDataRegistryEOS::FLobbyDataRegistryEOS(const TSharedRef<FLobbyPrerequisitesEOS>& Prerequisites)
+	: Prerequisites(Prerequisites)
+{
+}
+
+TSharedPtr<FLobbyDataEOS> FLobbyDataRegistryEOS::Find(EOS_LobbyId EOSLobbyId) const
+{
+	const TWeakPtr<FLobbyDataEOS>* Result = LobbyIdIndex.Find(Private::TranslateLobbyId(EOSLobbyId));
+	return Result ? Result->Pin() : TSharedPtr<FLobbyDataEOS>();
+}
+
+TSharedPtr<FLobbyDataEOS> FLobbyDataRegistryEOS::Find(FOnlineLobbyIdHandle LobbyIdHandle) const
+{
+	const TWeakPtr<FLobbyDataEOS>* Result = LobbyIdHandleIndex.Find(LobbyIdHandle);
+	return Result ? Result->Pin() : TSharedPtr<FLobbyDataEOS>();
+}
+
+TFuture<TDefaultErrorResultInternal<TSharedRef<FLobbyDataEOS>>> FLobbyDataRegistryEOS::FindOrCreateFromLobbyDetails(FOnlineAccountIdHandle LocalUserId, const TSharedRef<FLobbyDetailsEOS>& LobbyDetails)
+{
+	if (TSharedPtr<FLobbyDataEOS> FindResult = Find(LobbyDetails->GetInfo()->GetLobbyId()))
+	{
+		FindResult->AddUserLobbyDetails(LocalUserId, LobbyDetails);
+		return MakeFulfilledPromise<TDefaultErrorResultInternal<TSharedRef<FLobbyDataEOS>>>(FindResult.ToSharedRef()).GetFuture();
+	}
+
+	TPromise<TDefaultErrorResultInternal<TSharedRef<FLobbyDataEOS>>> Promise;
+	TFuture<TDefaultErrorResultInternal<TSharedRef<FLobbyDataEOS>>> Future = Promise.GetFuture();
+
+	const FOnlineLobbyIdHandle LobbyId = FOnlineLobbyIdHandle(EOnlineServices::Epic, NextHandleIndex++);
+	FLobbyDataEOS::Create(LobbyId, LobbyDetails, MakeUnregisterFn())
+	.Then([WeakThis = AsWeak(), Promise = MoveTemp(Promise), LocalUserId, LobbyDetails](TFuture<TDefaultErrorResultInternal<TSharedRef<FLobbyDataEOS>>>&& Future) mutable
+	{
+		if (TSharedPtr<FLobbyDataRegistryEOS> StrongThis = WeakThis.Pin())
+		{
+			if (Future.Get().IsOk())
+			{
+				StrongThis->Register(Future.Get().GetOkValue());
+			}
+		}
+
+		Future.Get().GetOkValue()->AddUserLobbyDetails(LocalUserId, LobbyDetails);
+		Promise.EmplaceValue(MoveTempIfPossible(Future.Get()));
+	});
+
+	return Future;
+}
+
+void FLobbyDataRegistryEOS::Register(const TSharedRef<FLobbyDataEOS>& LobbyIdHandleData)
+{
+	LobbyIdIndex.Add(LobbyIdHandleData->GetLobbyId(), LobbyIdHandleData);
+	LobbyIdHandleIndex.Add(LobbyIdHandleData->GetLobbyIdHandle(), LobbyIdHandleData);
+}
+
+void FLobbyDataRegistryEOS::Unregister(FOnlineLobbyIdHandle LobbyIdHandle)
+{
+	if (TSharedPtr<FLobbyDataEOS> HandleData = Find(LobbyIdHandle))
+	{
+		LobbyIdIndex.Remove(HandleData->GetLobbyId());
+		LobbyIdHandleIndex.Remove(HandleData->GetLobbyIdHandle());
+	}
+}
+
+FLobbyDataEOS::FUnregisterFn FLobbyDataRegistryEOS::MakeUnregisterFn()
+{
+	return [WeakThis = AsWeak()](FOnlineLobbyIdHandle LobbyId)
+	{
+		if (TSharedPtr<FLobbyDataRegistryEOS> StrongThis = WeakThis.Pin())
+		{
+			StrongThis->Unregister(LobbyId);
+		}
+	};
+}
+
+TFuture<TDefaultErrorResultInternal<TSharedRef<FLobbyInviteDataEOS>>> FLobbyInviteDataEOS::CreateFromInviteId(
+	const TSharedRef<FLobbyPrerequisitesEOS>& Prerequisites,
+	const TSharedRef<FLobbyDataRegistryEOS>& LobbyDataRegistry,
+	FOnlineAccountIdHandle LocalUserId,
+	const char* InviteIdEOS,
+	EOS_ProductUserId Sender)
+{
+	TDefaultErrorResultInternal<TSharedRef<FLobbyDetailsEOS>> LobbyDetailsResult = FLobbyDetailsEOS::CreateFromInviteId(Prerequisites, LocalUserId, InviteIdEOS);
+	if (LobbyDetailsResult.IsError())
+	{
+		return MakeFulfilledPromise<TDefaultErrorResultInternal<TSharedRef<FLobbyInviteDataEOS>>>(MoveTemp(LobbyDetailsResult.GetErrorValue())).GetFuture();
+	}
+
+	TSharedRef<FLobbyDetailsEOS> LobbyDetails = LobbyDetailsResult.GetOkValue();
+	TPromise<TDefaultErrorResultInternal<TSharedRef<FLobbyInviteDataEOS>>> Promise;
+	TFuture<TDefaultErrorResultInternal<TSharedRef<FLobbyInviteDataEOS>>> Future = Promise.GetFuture();
+
+	// Search for existing lobby data so that the LobbyIdHandle will match.
+	TSharedRef<FLobbyInviteIdEOS> InviteId = MakeShared<FLobbyInviteIdEOS>(InviteIdEOS);
+	LobbyDataRegistry->FindOrCreateFromLobbyDetails(LocalUserId, LobbyDetails)
+	.Then([Promise = MoveTemp(Promise), InviteId, LocalUserId, Sender, LobbyDetails](TFuture<TDefaultErrorResultInternal<TSharedRef<FLobbyDataEOS>>>&& Future) mutable
+	{
+		if (Future.Get().IsError())
+		{
+			Promise.EmplaceValue(MoveTemp(Future.Get().GetErrorValue()));
+			return;
+		}
+
+		// Once the lobby data has been resolved the FOnlineAccountIdHandle for the sender is expected to be in the AccountID cache.
+		const FOnlineAccountIdHandle SenderUserId = FindAccountId(Sender);
+		if (!SenderUserId.IsValid())
+		{
+			// Todo: Errors.
+			Promise.EmplaceValue(MoveTemp(Future.Get().GetErrorValue()));
+			return;
+		}
+
+		Promise.EmplaceValue(MakeShared<FLobbyInviteDataEOS>(InviteId, LocalUserId, SenderUserId, LobbyDetails, Future.Get().GetOkValue()));
+	});
+
+	return Future;
+}
+
+FLobbyInviteDataEOS::FLobbyInviteDataEOS(
+	const TSharedRef<FLobbyInviteIdEOS>& InviteIdEOS,
+	FOnlineAccountIdHandle Receiver,
+	FOnlineAccountIdHandle Sender,
+	const TSharedRef<FLobbyDetailsEOS>& LobbyDetails,
+	const TSharedRef<FLobbyDataEOS>& LobbyData)
+	: InviteIdEOS(InviteIdEOS)
+	, Receiver(Receiver)
+	, Sender(Sender)
+	, LobbyDetails(LobbyDetails)
+	, LobbyData(LobbyData)
+	, InviteId(Private::TranslateLobbyInviteId(InviteIdEOS->Get()))
+{
+}
+
+TFuture<TDefaultErrorResultInternal<TSharedRef<FLobbySearchEOS>>> FLobbySearchEOS::Create(
+	const TSharedRef<FLobbyPrerequisitesEOS>& Prerequisites,
+	const TSharedRef<FLobbyDataRegistryEOS>& LobbyRegistry,
+	const FLobbySearchParameters& Params)
+{
+	TSharedRef<FSearchHandle> SearchHandle = MakeShared<FSearchHandle>();
+
+	EOS_Lobby_CreateLobbySearchOptions CreateLobbySearchOptions = {};
+	CreateLobbySearchOptions.ApiVersion = EOS_LOBBY_CREATELOBBYSEARCH_API_LATEST;
+	CreateLobbySearchOptions.MaxResults = Params.MaxResults;
+
+	EOS_EResult EOSResult = EOS_Lobby_CreateLobbySearch(Prerequisites->LobbyInterfaceHandle, &CreateLobbySearchOptions, &SearchHandle->Get());
+	if (EOSResult != EOS_EResult::EOS_Success)
+	{
+		// todo: errors
+		return MakeFulfilledPromise<TDefaultErrorResultInternal<TSharedRef<FLobbySearchEOS>>>(FromEOSError(EOSResult)).GetFuture();
+	}
+
+	if (Params.LobbyId) // Search for specific lobby.
+	{
+		TSharedPtr<FLobbyDataEOS> LobbyData = LobbyRegistry->Find(*Params.LobbyId);
+		if (!LobbyData)
+		{
+			return MakeFulfilledPromise<TDefaultErrorResultInternal<TSharedRef<FLobbySearchEOS>>>(Errors::InvalidParams()).GetFuture();
+		}
+
+		EOS_LobbySearch_SetLobbyIdOptions SetLobbyIdOptions = {};
+		SetLobbyIdOptions.ApiVersion = EOS_LOBBYSEARCH_SETLOBBYID_API_LATEST;
+		SetLobbyIdOptions.LobbyId = LobbyData->GetLobbyIdEOS();
+
+		EOSResult = EOS_LobbySearch_SetLobbyId(SearchHandle->Get(), &SetLobbyIdOptions);
+		if (EOSResult != EOS_EResult::EOS_Success)
+		{
+			// todo: errors
+			return MakeFulfilledPromise<TDefaultErrorResultInternal<TSharedRef<FLobbySearchEOS>>>(FromEOSError(EOSResult)).GetFuture();
+		}
+	}
+	else if (Params.TargetUser) // Search for specific user.
+	{
+		EOS_LobbySearch_SetTargetUserIdOptions SetTargetUserIdOptions = {};
+		SetTargetUserIdOptions.ApiVersion = EOS_LOBBYSEARCH_SETTARGETUSERID_API_LATEST;
+		SetTargetUserIdOptions.TargetUserId = GetProductUserIdChecked(*Params.TargetUser);
+
+		EOSResult = EOS_LobbySearch_SetTargetUserId(SearchHandle->Get(), &SetTargetUserIdOptions);
+		if (EOSResult != EOS_EResult::EOS_Success)
+		{
+			// todo: errors
+			return MakeFulfilledPromise<TDefaultErrorResultInternal<TSharedRef<FLobbySearchEOS>>>(FromEOSError(EOSResult)).GetFuture();
+		}
+	}
+	else // Search using parameters.
+	{
+		// Bucket id.
+		{
+			const FLobbyBucketIdTranslator<ELobbyTranslationType::ToService> BucketTranslator(Prerequisites->BucketId);
+
+			EOS_Lobby_AttributeData AttributeData;
+			AttributeData.ApiVersion = EOS_LOBBY_ATTRIBUTEDATA_API_LATEST;
+			AttributeData.Key = EOS_LOBBY_SEARCH_BUCKET_ID;
+			AttributeData.ValueType = EOS_ELobbyAttributeType::EOS_AT_STRING;
+			AttributeData.Value.AsUtf8 = BucketTranslator.GetBucketIdEOS();
+
+			EOS_LobbySearch_SetParameterOptions SetParameterOptions = {};
+			SetParameterOptions.ApiVersion = EOS_LOBBYSEARCH_SETTARGETUSERID_API_LATEST;
+			SetParameterOptions.Parameter = &AttributeData;
+			SetParameterOptions.ComparisonOp = EOS_EComparisonOp::EOS_CO_EQUAL;
+
+			EOSResult = EOS_LobbySearch_SetParameter(SearchHandle->Get(), &SetParameterOptions);
+			if (EOSResult != EOS_EResult::EOS_Success)
+			{
+				// todo: errors
+				return MakeFulfilledPromise<TDefaultErrorResultInternal<TSharedRef<FLobbySearchEOS>>>(FromEOSError(EOSResult)).GetFuture();
+			}
+		}
+
+		for (const FFindLobbySearchFilter& Filter :  Params.Filters)
+		{
+			const FLobbyAttributeTranslator<ELobbyTranslationType::ToService> AttributeTranslator(
+				Filter.AttributeName, Filter.ComparisonValue);
+
+			EOS_LobbySearch_SetParameterOptions SetParameterOptions = {};
+			SetParameterOptions.ApiVersion = EOS_LOBBYSEARCH_SETTARGETUSERID_API_LATEST;
+			SetParameterOptions.Parameter = &AttributeTranslator.GetAttributeData();
+			SetParameterOptions.ComparisonOp = Private::TranslateSearchComparison(Filter.ComparisonOp);
+
+			EOSResult = EOS_LobbySearch_SetParameter(SearchHandle->Get(), &SetParameterOptions);
+			if (EOSResult != EOS_EResult::EOS_Success)
+			{
+				// todo: errors
+				return MakeFulfilledPromise<TDefaultErrorResultInternal<TSharedRef<FLobbySearchEOS>>>(FromEOSError(EOSResult)).GetFuture();
+			}
+		}
+	}
+
+	TPromise<TDefaultErrorResultInternal<TSharedRef<FLobbySearchEOS>>> Promise;
+	TFuture<TDefaultErrorResultInternal<TSharedRef<FLobbySearchEOS>>> Future = Promise.GetFuture();
+
+	EOS_LobbySearch_FindOptions FindOptions = {};
+	FindOptions.ApiVersion = EOS_LOBBYSEARCH_FIND_API_LATEST;
+	FindOptions.LocalUserId = GetProductUserIdChecked(Params.LocalUserId);
+
+	EOS_Async<EOS_LobbySearch_FindCallbackInfo>(EOS_LobbySearch_Find, SearchHandle->Get(), FindOptions)
+	.Then([Promise = MoveTemp(Promise), Prerequisites, LobbyRegistry, LocalUserId = Params.LocalUserId, SearchHandle]
+	(TFuture<const EOS_LobbySearch_FindCallbackInfo*>&& Future) mutable
+	{
+		if (Future.Get()->ResultCode != EOS_EResult::EOS_Success)
+		{
+			// todo: errors
+			Promise.EmplaceValue(FromEOSError(Future.Get()->ResultCode));
+			return;
+		}
+
+		TArray<TFuture<TDefaultErrorResultInternal<TSharedRef<FLobbyDataEOS>>>> ResolvedLobbyDetails;
+
+		EOS_LobbySearch_GetSearchResultCountOptions GetSearchResultCountOptions = {};
+		GetSearchResultCountOptions.ApiVersion = EOS_LOBBYSEARCH_GETSEARCHRESULTCOUNT_API_LATEST;
+		const uint32_t NumSearchResults = EOS_LobbySearch_GetSearchResultCount(SearchHandle->Get(), &GetSearchResultCountOptions);
+
+		for (uint32_t SearchResultIndex = 0; SearchResultIndex < NumSearchResults; ++SearchResultIndex)
+		{
+			TDefaultErrorResultInternal<TSharedRef<FLobbyDetailsEOS>> Result = FLobbyDetailsEOS::CreateFromSearchResult(
+				Prerequisites, LocalUserId, SearchHandle->Get(), SearchResultIndex);
+			if (Result.IsError())
+			{
+				// todo: errors
+				Promise.EmplaceValue(MoveTemp(Result.GetErrorValue()));
+				return;
+			}
+
+			TPromise<TDefaultErrorResultInternal<TSharedRef<FLobbyDataEOS>>> ResolveLobbyDetailsPromise;
+			ResolvedLobbyDetails.Add(ResolveLobbyDetailsPromise.GetFuture());
+
+			LobbyRegistry->FindOrCreateFromLobbyDetails(LocalUserId, Result.GetOkValue())
+			.Then([ResolveLobbyDetailsPromise = MoveTemp(ResolveLobbyDetailsPromise)](TFuture<TDefaultErrorResultInternal<TSharedRef<FLobbyDataEOS>>>&& Future) mutable
+			{
+				ResolveLobbyDetailsPromise.EmplaceValue(MoveTempIfPossible(Future.Get()));
+			});
+		}
+
+		WhenAll(MoveTemp(ResolvedLobbyDetails))
+		.Then([Promise = MoveTemp(Promise), SearchHandle](TFuture<TArray<TDefaultErrorResultInternal<TSharedRef<FLobbyDataEOS>>>>&& Future) mutable
+		{
+			TArray<TDefaultErrorResultInternal<TSharedRef<FLobbyDataEOS>>> Results(Future.Get());
+			TArray<TSharedRef<FLobbyDataEOS>> ResolvedResults;
+			ResolvedResults.Reserve(Results.Num());
+
+			for (TDefaultErrorResultInternal<TSharedRef<FLobbyDataEOS>>& Result : Results)
+			{
+				if (Result.IsError())
+				{
+					// todo: errors
+					Promise.EmplaceValue(MoveTemp(Result.GetErrorValue()));
+					return;
+				}
+
+				ResolvedResults.Add(MoveTemp(Result.GetOkValue()));
+			}
+
+			Promise.EmplaceValue(MakeShared<FLobbySearchEOS>(SearchHandle, MoveTemp(ResolvedResults)));
+		});
+	});
+
+	return Future;
+}
+
+TArray<TSharedRef<const FLobby>> FLobbySearchEOS::GetLobbyResults() const
+{
+	TArray<TSharedRef<const FLobby>> Result;
+	Result.Reserve(Lobbies.Num());
+
+	for (const TSharedRef<FLobbyDataEOS>& LobbyData : Lobbies)
+	{
+		Result.Add(LobbyData->GetClientLobbyData()->GetPublicDataPtr());
+	}
+
+	return Result;
+}
+
+const TArray<TSharedRef<FLobbyDataEOS>>& FLobbySearchEOS::GetLobbyData()
+{
+	return Lobbies;
+}
+
+FLobbySearchEOS::FLobbySearchEOS(const TSharedRef<FSearchHandle>& SearchHandle, TArray<TSharedRef<FLobbyDataEOS>>&& Lobbies)
+	: SearchHandle(SearchHandle)
+	, Lobbies(Lobbies)
+{
+}
+
+/* UE::Online */ }
\ No newline at end of file
diff --git a/Engine/Plugins/Online/Experimental/OnlineServicesEOSGS/Source/Private/Online/LobbiesEOSGSTypes.h b/Engine/Plugins/Online/Experimental/OnlineServicesEOSGS/Source/Private/Online/LobbiesEOSGSTypes.h
new file mode 100644
--- /dev/null
+++ b/Engine/Plugins/Online/Experimental/OnlineServicesEOSGS/Source/Private/Online/LobbiesEOSGSTypes.h
@@ -0,0 +1,488 @@
+// Copyright Epic Games, Inc. All Rights Reserved.
+
+#pragma once
+
+#include "Online/Lobbies.h"
+#include "Online/LobbiesCommonTypes.h"
+#include "Online/OnlineIdEOSGS.h"
+#include "Online/OnlineServicesEOSGSTypes.h"
+#include "CoreMinimal.h"
+
+#if defined(EOS_PLATFORM_BASE_FILE_NAME)
+#include EOS_PLATFORM_BASE_FILE_NAME
+#endif
+
+#include "eos_lobby.h"
+
+namespace UE::Online {
+
+class FLobbySchemaRegistry;
+class FLobbyDataRegistryEOS;
+class FLobbyDetailsInfoEOS;
+class FLobbyInviteDataEOSRegistry;
+class FLobbySchema;
+struct FLobbyEvents;
+class FAuthEOS;
+struct FClientLobbyMemberSnapshot;
+
+//--------------------------------------------------------------------------------------------------
+// Translated types
+//--------------------------------------------------------------------------------------------------
+
+class FLobbyBucketIdEOS
+{
+public:
+	static const FString Separator;
+
+	FLobbyBucketIdEOS() = default;
+	FLobbyBucketIdEOS(const FLobbyBucketIdEOS&) = default;
+	FLobbyBucketIdEOS(FLobbyBucketIdEOS&&) = default;
+	FLobbyBucketIdEOS& operator=(const FLobbyBucketIdEOS&) = default;
+	FLobbyBucketIdEOS& operator=(FLobbyBucketIdEOS&&) = default;
+
+	FLobbyBucketIdEOS(FString ProductName, int32 ProductVersion);
+
+	const FString& GetProductName() const { return ProductName; }
+	int32 GetProductVersion() const { return ProductVersion; }
+	bool IsValid() const { return !ProductName.IsEmpty(); }
+
+private:
+	FString ProductName;
+	int32 ProductVersion;
+};
+
+//--------------------------------------------------------------------------------------------------
+// Translators
+//--------------------------------------------------------------------------------------------------
+enum class ELobbyTranslationType
+{
+	ToService,
+	FromService
+};
+
+//--------------------------------------------------------------------------------------------------
+// Lobby attribute
+
+template <ELobbyTranslationType>
+class FLobbyAttributeTranslator
+{
+public:
+};
+
+template <>
+class FLobbyAttributeTranslator<ELobbyTranslationType::ToService>
+{
+public:
+	FLobbyAttributeTranslator(const TPair<FLobbyAttributeId, FLobbyVariant>& FromAttributeData);
+	FLobbyAttributeTranslator(FLobbyAttributeId FromAttributeId, const FLobbyVariant& FromAttributeData);
+
+	const EOS_Lobby_AttributeData& GetAttributeData() const { return AttributeData; }
+
+private:
+	FTCHARToUTF8 KeyConverterStorage;
+	TOptional<FTCHARToUTF8> ValueConverterStorage;
+	EOS_Lobby_AttributeData AttributeData;
+};
+
+template <>
+class FLobbyAttributeTranslator<ELobbyTranslationType::FromService>
+{
+public:
+	FLobbyAttributeTranslator(const EOS_Lobby_AttributeData& FromAttributeData);
+
+	const TPair<FLobbyAttributeId, FLobbyVariant>& GetAttributeData() const { return AttributeData; }
+	TPair<FLobbyAttributeId, FLobbyVariant>& GetMutableAttributeData() { return AttributeData; }
+
+private:
+	TPair<FLobbyAttributeId, FLobbyVariant> AttributeData;
+};
+
+//--------------------------------------------------------------------------------------------------
+// Bucket id
+
+template <ELobbyTranslationType>
+class FLobbyBucketIdTranslator;
+
+template <>
+class FLobbyBucketIdTranslator<ELobbyTranslationType::ToService>
+{
+public:
+	FLobbyBucketIdTranslator(const FLobbyBucketIdEOS& BucketId);
+
+	const char* GetBucketIdEOS() const { return BucketConverterStorage.Get(); }
+
+private:
+	FTCHARToUTF8 BucketConverterStorage;
+};
+
+template <>
+class FLobbyBucketIdTranslator<ELobbyTranslationType::FromService>
+{
+public:
+	FLobbyBucketIdTranslator(const char* BucketIdEOS);
+
+	const FLobbyBucketIdEOS& GetBucketId() const { return BucketId; }
+	FLobbyBucketIdEOS& GetMutableBucketId() { return BucketId; }
+
+private:
+	FLobbyBucketIdEOS BucketId;
+};
+
+//--------------------------------------------------------------------------------------------------
+// Enumerations
+
+inline EOS_ELobbyPermissionLevel TranslateJoinPolicy(ELobbyJoinPolicy JoinPolicy)
+{
+	switch (JoinPolicy)
+	{
+	case ELobbyJoinPolicy::PublicAdvertised:	return EOS_ELobbyPermissionLevel::EOS_LPL_PUBLICADVERTISED;
+	case ELobbyJoinPolicy::PublicNotAdvertised:	return EOS_ELobbyPermissionLevel::EOS_LPL_JOINVIAPRESENCE;
+	default:									checkNoEntry(); // Intentional fallthrough
+	case ELobbyJoinPolicy::InvitationOnly:		return EOS_ELobbyPermissionLevel::EOS_LPL_INVITEONLY;
+	};
+}
+
+inline ELobbyJoinPolicy TranslateJoinPolicy(EOS_ELobbyPermissionLevel JoinPolicy)
+{
+	switch (JoinPolicy)
+	{
+	case EOS_ELobbyPermissionLevel::EOS_LPL_PUBLICADVERTISED:	return ELobbyJoinPolicy::PublicAdvertised;
+	case EOS_ELobbyPermissionLevel::EOS_LPL_JOINVIAPRESENCE:	return ELobbyJoinPolicy::PublicNotAdvertised;
+	default:									checkNoEntry(); // Intentional fallthrough
+	case EOS_ELobbyPermissionLevel::EOS_LPL_INVITEONLY:	return ELobbyJoinPolicy::InvitationOnly;
+	};
+}
+
+//--------------------------------------------------------------------------------------------------
+// Structures.
+//--------------------------------------------------------------------------------------------------
+
+using FLobbySearchParameters = FFindLobbies::Params;
+
+/**
+* Common components required to handle lobby requests.
+* To handle lifetime issues, only FLobbiesEOS should contain a strong reference to prerequisites.
+* Any other component should only store a weak reference.
+*/
+struct FLobbyPrerequisitesEOS
+{
+	EOS_HLobby LobbyInterfaceHandle = {};
+	TWeakPtr<FAuthEOSGS> AuthInterface;
+	TSharedRef<const FLobbySchemaRegistry> SchemaRegistry;
+	TSharedRef<const FLobbySchema> ServiceSchema;
+	FLobbyBucketIdEOS BucketId;
+};
+
+/**
+ * Lobby details are created for each user within the EOS lobby client. Certain operations such as
+ * joining a lobby or applying an update to a lobby / member attributes require using the correct
+ * lobby details handle.
+ */
+enum class ELobbyDetailsSource
+{
+	/**
+	 * The lobby has been joined and is considered active.
+	 * Valid for lobby / member updates.
+	 */
+	Active,
+
+	/**
+	 * The details originated from an invitation sent to the user.
+	 * Valid for joining.
+	 */
+	Invite,
+
+	/**
+	 * The details originated from an event generated by the EOS overlay.
+	 * Valid for joining.
+	 */
+	UiEvent,
+
+	/**
+	 * The details originated from a lobby search result.
+	 * Valid for joining.
+	 */
+	Search
+};
+
+/** Lobby details is created based on the passed in user and is required to join a lobby. */
+class FLobbyDetailsEOS final : public TSharedFromThis<FLobbyDetailsEOS>
+{
+public:
+	UE_NONCOPYABLE(FLobbyDetailsEOS);
+
+	static TDefaultErrorResultInternal<TSharedRef<FLobbyDetailsEOS>> CreateFromLobbyId(
+		const TSharedRef<FLobbyPrerequisitesEOS>& Prerequisites,
+		FOnlineAccountIdHandle LocalUserId,
+		EOS_LobbyId LobbyId);
+	static TDefaultErrorResultInternal<TSharedRef<FLobbyDetailsEOS>> CreateFromInviteId(
+		const TSharedRef<FLobbyPrerequisitesEOS>& Prerequisites,
+		FOnlineAccountIdHandle LocalUserId,
+		const char* InviteId);
+	static TDefaultErrorResultInternal<TSharedRef<FLobbyDetailsEOS>> CreateFromUiEventId(
+		const TSharedRef<FLobbyPrerequisitesEOS>& Prerequisites,
+		FOnlineAccountIdHandle LocalUserId,
+		EOS_UI_EventId UiEventId);
+	static TDefaultErrorResultInternal<TSharedRef<FLobbyDetailsEOS>> CreateFromSearchResult(
+		const TSharedRef<FLobbyPrerequisitesEOS>& Prerequisites,
+		FOnlineAccountIdHandle LocalUserId,
+		EOS_HLobbySearch SearchHandle,
+		uint32_t ResultIndex);
+
+	FLobbyDetailsEOS() = delete;
+	~FLobbyDetailsEOS();
+
+	EOS_HLobbyDetails GetEOSHandle() const { return LobbyDetailsHandle; }
+	const TSharedRef<FLobbyDetailsInfoEOS>& GetInfo() const { return LobbyDetailsInfo; }
+	ELobbyDetailsSource GetDetailsSource() const { return LobbyDetailsSource; }
+	FOnlineAccountIdHandle GetAssociatedUser() const { return AssociatedLocalUser; }
+
+	/**
+	 * Retrieve a lobby data snapshot from the EOS lobby details object.
+	 * The lobby schema will be used to translate attribute data before returning.
+	 */
+	TFuture<TDefaultErrorResultInternal<TSharedRef<FClientLobbySnapshot>>> GetLobbySnapshot() const;
+
+	/**
+	 * Retrieve lobby member data snapshot from the EOS lobby details object.
+	 * The lobby schema will be used to translate attribute data before returning.
+	 */
+	TDefaultErrorResultInternal<TSharedRef<FClientLobbyMemberSnapshot>> GetLobbyMemberSnapshot(FOnlineAccountIdHandle MemberHandle) const;
+
+	/**
+	 * Apply client side lobby changes to the lobby service.
+	 * The lobby schema will be used to translate any changed attributes before sending.
+	 */
+	TFuture<EOS_EResult> ApplyLobbyDataUpdateFromLocalChanges(FOnlineAccountIdHandle LocalUserId, const FClientLobbyDataChanges& Changes) const;
+
+	/**
+	 * Apply client side lobby member changes to the lobby service.
+	 * The lobby schema will be used to translate any changed attributes before sending.
+	 */
+	TFuture<EOS_EResult> ApplyLobbyMemberDataUpdateFromLocalChanges(FOnlineAccountIdHandle LocalUserId, const FClientLobbyMemberDataChanges& Changes) const;
+
+private:
+	template <typename, ESPMode>
+	friend class SharedPointerInternals::TIntrusiveReferenceController;
+	FLobbyDetailsEOS(
+		const TSharedRef<FLobbyPrerequisitesEOS>& Prerequisites,
+		const TSharedRef<FLobbyDetailsInfoEOS>& LobbyDetailsInfo,
+		FOnlineAccountIdHandle LocalUserId,
+		ELobbyDetailsSource LobbyDetailsSource,
+		EOS_HLobbyDetails LobbyDetailsHandle);
+
+	static const EOS_HLobbyDetails InvalidLobbyDetailsHandle;
+	TSharedRef<FLobbyPrerequisitesEOS> Prerequisites;
+	TSharedRef<FLobbyDetailsInfoEOS> LobbyDetailsInfo;
+	FOnlineAccountIdHandle AssociatedLocalUser;
+	ELobbyDetailsSource LobbyDetailsSource;
+	EOS_HLobbyDetails LobbyDetailsHandle = InvalidLobbyDetailsHandle;
+};
+
+/** Lobby details info allows direct access to some of the lobby properties. */
+class FLobbyDetailsInfoEOS final
+{
+public:
+	UE_NONCOPYABLE(FLobbyDetailsInfoEOS);
+	FLobbyDetailsInfoEOS() = delete;
+
+	static TDefaultErrorResultInternal<TSharedRef<FLobbyDetailsInfoEOS>> Create(EOS_HLobbyDetails LobbyDetailsHandle);
+
+	EOS_LobbyId GetLobbyId() const { return LobbyDetailsInfo->LobbyId; }
+	int32 GetMaxMembers() const { return LobbyDetailsInfo->MaxMembers; }
+	EOS_ELobbyPermissionLevel GetPermissionLevel() const { return LobbyDetailsInfo->PermissionLevel; }
+
+	const FString& GetProductName() const { return BucketId.GetProductName(); }
+	int32 GetProductVersion() const { return BucketId.GetProductVersion(); }
+
+private:
+	struct FEOSLobbyDetailsInfoDeleter
+	{
+		void operator()(EOS_LobbyDetails_Info* Ptr) const
+		{
+			EOS_LobbyDetails_Info_Release(Ptr);
+		}
+	};
+
+	using FLobbyDetailsInfoPtr = TUniquePtr<EOS_LobbyDetails_Info, FEOSLobbyDetailsInfoDeleter>;
+
+	template <typename, ESPMode>
+	friend class SharedPointerInternals::TIntrusiveReferenceController;
+	FLobbyDetailsInfoEOS(FLobbyDetailsInfoPtr&& LobbyDetailsInfo);
+
+	TUniquePtr<EOS_LobbyDetails_Info, FEOSLobbyDetailsInfoDeleter> LobbyDetailsInfo;
+	FLobbyBucketIdEOS BucketId;
+};
+
+/** Lobby data is the bookkeeping object for a lobby. It contains the client-side representation of a lobby. */
+class FLobbyDataEOS final : public TSharedFromThis<FLobbyDataEOS>
+{
+public:
+	using FUnregisterFn = TFunction<void(FOnlineLobbyIdHandle)>;
+
+	FLobbyDataEOS() = delete;
+
+	static TFuture<TDefaultErrorResultInternal<TSharedRef<FLobbyDataEOS>>> Create(
+		FOnlineLobbyIdHandle LobbyIdHandle,
+		const TSharedRef<FLobbyDetailsEOS>& LobbyDetails,
+		FUnregisterFn UnregisterFn = FUnregisterFn());
+
+	~FLobbyDataEOS();
+
+	FOnlineLobbyIdHandle GetLobbyIdHandle() const { return ClientLobbyData->GetPublicData().LobbyId; }
+	const TSharedRef<FClientLobbyData>& GetClientLobbyData() const { return ClientLobbyData; }
+	EOS_LobbyId GetLobbyIdEOS() const { return LobbyDetailsInfo->GetLobbyId(); }
+	const FString& GetLobbyId() const { return LobbyId; }
+
+	void AddUserLobbyDetails(FOnlineAccountIdHandle LocalUserId, const TSharedPtr<FLobbyDetailsEOS>& LobbyDetails);
+	TSharedPtr<FLobbyDetailsEOS> GetUserLobbyDetails(FOnlineAccountIdHandle LocalUserId) const;
+
+	/**
+	 * Active lobby details are needed to process lobby notifications. Search for and return active
+	 * lobby details if available.
+	 */
+	TSharedPtr<FLobbyDetailsEOS> GetActiveLobbyDetails() const;
+
+private:
+	template <typename, ESPMode>
+	friend class SharedPointerInternals::TIntrusiveReferenceController;
+
+	FLobbyDataEOS(
+		const TSharedRef<FClientLobbyData>& ClientLobbyData,
+		const TSharedRef<FLobbyDetailsInfoEOS>& LobbyDetailsInfo,
+		FUnregisterFn UnregisterFn);
+
+	TSharedRef<FClientLobbyData> ClientLobbyData;
+	TSharedRef<FLobbyDetailsInfoEOS> LobbyDetailsInfo;
+	FUnregisterFn UnregisterFn;
+	FString LobbyId;
+	TMap<FOnlineAccountIdHandle, TSharedPtr<FLobbyDetailsEOS>> UserLobbyDetails;
+};
+
+// Todo: implement handling for FOnlineIdRegistryRegistry.
+class FLobbyDataRegistryEOS : public TSharedFromThis<FLobbyDataRegistryEOS>
+{
+public:
+	FLobbyDataRegistryEOS(const TSharedRef<FLobbyPrerequisitesEOS>& Prerequisites);
+
+	TSharedPtr<FLobbyDataEOS> Find(EOS_LobbyId EOSLobbyId) const;
+	TSharedPtr<FLobbyDataEOS> Find(FOnlineLobbyIdHandle LobbyIdHandle) const;
+	TFuture<TDefaultErrorResultInternal<TSharedRef<FLobbyDataEOS>>> FindOrCreateFromLobbyDetails(FOnlineAccountIdHandle LocalUserId, const TSharedRef<FLobbyDetailsEOS>& LobbyDetails);
+
+private:
+	void Register(const TSharedRef<FLobbyDataEOS>& LobbyIdHandleData);
+	void Unregister(FOnlineLobbyIdHandle LobbyIdHandle);
+	FLobbyDataEOS::FUnregisterFn MakeUnregisterFn();
+
+	TSharedRef<FLobbyPrerequisitesEOS> Prerequisites;
+	TMap<FString, TWeakPtr<FLobbyDataEOS>> LobbyIdIndex;
+	TMap<FOnlineLobbyIdHandle, TWeakPtr<FLobbyDataEOS>> LobbyIdHandleIndex;
+	uint32 NextHandleIndex = 1;
+};
+
+class FLobbyInviteIdEOS final
+{
+public:
+	FLobbyInviteIdEOS() = default;
+	FLobbyInviteIdEOS(const char* InviteIdEOS)
+	{
+		FPlatformString::Strcpy(Data, MaxLobbyInviteIdSize, InviteIdEOS);
+	}
+
+	const char* Get() const { return Data; }
+
+private:
+	static const int MaxLobbyInviteIdSize = 256;
+	char Data[MaxLobbyInviteIdSize] = {};
+};
+
+/**
+* Lobby invite data will keep the lobby object valid until the invitation has been accepted or rejected.
+*/
+class FLobbyInviteDataEOS final
+{
+public:
+	FLobbyInviteDataEOS() = delete;
+
+	static TFuture<TDefaultErrorResultInternal<TSharedRef<FLobbyInviteDataEOS>>> CreateFromInviteId(
+		const TSharedRef<FLobbyPrerequisitesEOS>& Prerequisites,
+		const TSharedRef<FLobbyDataRegistryEOS>& LobbyDataRegistry,
+		FOnlineAccountIdHandle LocalUserId,
+		const char* InviteIdEOS,
+		EOS_ProductUserId Sender);
+
+	TSharedRef<FLobbyDetailsEOS> GetLobbyDetails() const { return LobbyDetails; }
+	TSharedRef<FLobbyDataEOS> GetLobbyData() const { return LobbyData; }
+
+	const char* GetInviteIdEOS() const { return InviteIdEOS->Get(); }
+	const FString& GetInviteId() const { return InviteId; }
+
+	FOnlineAccountIdHandle GetReceiver() const { return Receiver; }
+	FOnlineAccountIdHandle GetSender() const { return Sender; }
+
+private:
+	template <typename, ESPMode>
+	friend class SharedPointerInternals::TIntrusiveReferenceController;
+
+	FLobbyInviteDataEOS(
+		const TSharedRef<FLobbyInviteIdEOS>& InviteIdEOS,
+		FOnlineAccountIdHandle Receiver,
+		FOnlineAccountIdHandle Sender,
+		const TSharedRef<FLobbyDetailsEOS>& LobbyDetails,
+		const TSharedRef<FLobbyDataEOS>& LobbyData);
+
+	TSharedRef<FLobbyInviteIdEOS> InviteIdEOS;
+	FOnlineAccountIdHandle Receiver;
+	FOnlineAccountIdHandle Sender;
+	TSharedRef<FLobbyDetailsEOS> LobbyDetails;
+	TSharedRef<FLobbyDataEOS> LobbyData;
+	FString InviteId;
+};
+
+/**
+* The lobby search object is meant to hold the lifetimes of lobby search results so that the
+* lobby details will be discoverable during a join operation.
+*/
+class FLobbySearchEOS final
+{
+public:
+	FLobbySearchEOS() = delete;
+
+	static TFuture<TDefaultErrorResultInternal<TSharedRef<FLobbySearchEOS>>> Create(
+		const TSharedRef<FLobbyPrerequisitesEOS>& Prerequisites,
+		const TSharedRef<FLobbyDataRegistryEOS>& LobbyRegistry,
+		const FLobbySearchParameters& Params);
+
+	TArray<TSharedRef<const FLobby>> GetLobbyResults() const;
+	const TArray<TSharedRef<FLobbyDataEOS>>& GetLobbyData();
+
+private:
+	class FSearchHandle
+	{
+	public:
+		UE_NONCOPYABLE(FSearchHandle);
+		FSearchHandle() = default;
+
+		~FSearchHandle()
+		{
+			EOS_LobbySearch_Release(SearchHandle);
+		}
+
+		EOS_HLobbySearch& Get()
+		{
+			return SearchHandle;
+		}
+
+	private:
+		EOS_HLobbySearch SearchHandle = {};
+	};
+
+	template <typename, ESPMode>
+	friend class SharedPointerInternals::TIntrusiveReferenceController;
+	FLobbySearchEOS(const TSharedRef<FSearchHandle>& SearchHandle, TArray<TSharedRef<FLobbyDataEOS>>&& Lobbies);
+
+	TSharedRef<FSearchHandle> SearchHandle;
+	TArray<TSharedRef<FLobbyDataEOS>> Lobbies;
+};
+
+/* UE::Online */ }
diff --git a/Engine/Plugins/Online/Experimental/OnlineServicesEOSGS/Source/Private/Online/OnlineErrorEOSGS.cpp b/Engine/Plugins/Online/Experimental/OnlineServicesEOSGS/Source/Private/Online/OnlineErrorEOSGS.cpp
new file mode 100644
--- /dev/null
+++ b/Engine/Plugins/Online/Experimental/OnlineServicesEOSGS/Source/Private/Online/OnlineErrorEOSGS.cpp
@@ -0,0 +1,40 @@
+// Copyright Epic Games, Inc. All Rights Reserved.
+
+#include "OnlineErrorEOSGS.h"
+
+namespace UE::Online {
+
+const TSharedRef<const IOnlineErrorDetails, ESPMode::ThreadSafe>& FOnlineErrorDetailsEOS::Get()
+{
+	static TSharedRef<const IOnlineErrorDetails, ESPMode::ThreadSafe> Instance = MakeShared<FOnlineErrorDetailsEOS>();
+	return Instance;
+}
+
+FText FOnlineErrorDetailsEOS::GetText(const FOnlineError& OnlineError) const
+{
+	const EOS_EResult EosResult = EOS_EResult(OnlineError.GetValue());
+	return FText::FromStringView(FStringView(UTF8_TO_TCHAR(EOS_EResult_ToString(EosResult))));
+}
+
+FString FOnlineErrorDetailsEOS::GetLogString(const FOnlineError& OnlineError) const 
+{
+	const EOS_EResult EosResult = EOS_EResult(OnlineError.GetValue());
+	return UTF8_TO_TCHAR(EOS_EResult_ToString(EosResult));
+}
+
+bool operator==(const UE::Online::FOnlineError& OnlineError, EOS_EResult EosResult)
+{
+	return OnlineError.GetErrorCode() == CreateErrorCode(EosResult);
+}
+
+ErrorCodeType CreateErrorCode(EOS_EResult EosResult)
+{
+	return Errors::ErrorCode::Create(Errors::ErrorCode::Category::EOS_System, Errors::ErrorCode::Category::EOS, uint32(EosResult));
+}
+
+FOnlineError CreateOnlineError(EOS_EResult EosResult)
+{
+	return FOnlineError(CreateErrorCode(EosResult), FOnlineErrorDetailsEOS::Get());
+}
+
+} /* namespace UE::Online */
diff --git a/Engine/Plugins/Online/Experimental/OnlineServicesEOSGS/Source/Private/Online/OnlineErrorEOSGS.h b/Engine/Plugins/Online/Experimental/OnlineServicesEOSGS/Source/Private/Online/OnlineErrorEOSGS.h
new file mode 100644
--- /dev/null
+++ b/Engine/Plugins/Online/Experimental/OnlineServicesEOSGS/Source/Private/Online/OnlineErrorEOSGS.h
@@ -0,0 +1,37 @@
+// Copyright Epic Games, Inc. All Rights Reserved.
+
+#pragma once
+
+#include "Online/OnlineError.h"
+
+#if defined(EOS_PLATFORM_BASE_FILE_NAME)
+#include EOS_PLATFORM_BASE_FILE_NAME
+#endif
+#include "eos_common.h"
+
+namespace UE::Online {
+
+namespace Errors {
+
+UE_ONLINE_ERROR_CATEGORY(EOS, Engine, 0x3, "Online Services EOS")
+
+} /* namespace Errors */
+
+class FOnlineErrorDetailsEOS : public IOnlineErrorDetails
+{
+public:
+	static const TSharedRef<const IOnlineErrorDetails, ESPMode::ThreadSafe>& Get();
+
+	virtual FText GetText(const FOnlineError& OnlineError) const override;
+	virtual FString GetLogString(const FOnlineError& OnlineError) const override;
+};
+
+bool operator==(const FOnlineError& OnlineError, EOS_EResult EosResult);
+inline bool operator==(EOS_EResult EosResult, const FOnlineError& OnlineError) { return OnlineError == EosResult; }
+inline bool operator!=(const FOnlineError& OnlineError, EOS_EResult EosResult) { return !(OnlineError == EosResult); }
+inline bool operator!=(EOS_EResult EosResult, const FOnlineError& OnlineError) { return !(OnlineError == EosResult); }
+
+ErrorCodeType CreateErrorCode(EOS_EResult EosResult);
+FOnlineError CreateOnlineError(EOS_EResult EosResult);
+
+} /* namespace UE::Online */
\ No newline at end of file
diff --git a/Engine/Plugins/Online/Experimental/OnlineServicesEOSGS/Source/Private/Online/OnlineIdEOSGS.cpp b/Engine/Plugins/Online/Experimental/OnlineServicesEOSGS/Source/Private/Online/OnlineIdEOSGS.cpp
new file mode 100644
--- /dev/null
+++ b/Engine/Plugins/Online/Experimental/OnlineServicesEOSGS/Source/Private/Online/OnlineIdEOSGS.cpp
@@ -0,0 +1,114 @@
+// Copyright Epic Games, Inc. All Rights Reserved.
+
+#include "Online/OnlineIdEOSGS.h"
+
+#include "Online/OnlineAsyncOp.h"
+#include "String/HexToBytes.h"
+#include "String/BytesToHex.h"
+#include "Online/OnlineServicesEOSGSTypes.h"
+
+namespace UE::Online {
+
+const uint8 OnlineIdEOSUtf8BufferLength = 32;
+const uint8 OnlineIdEOSHexBufferLength = 16;
+
+FOnlineAccountIdHandle FOnlineAccountIdRegistryEOSGS::FindAccountId(const EOS_ProductUserId ProductUserId) const
+{
+	return Registry.FindHandle(ProductUserId);
+}
+
+EOS_ProductUserId FOnlineAccountIdRegistryEOSGS::GetProductUserId(const FOnlineAccountIdHandle& Handle) const
+{
+	return Registry.FindIdValue(Handle);
+}
+
+FString FOnlineAccountIdRegistryEOSGS::ToLogString(const FOnlineAccountIdHandle& Handle) const
+{
+	if (Registry.ValidateOnlineId(Handle))
+	{
+		EOS_ProductUserId ProductUserId = Registry.FindIdValue(Handle);
+		return LexToString(ProductUserId);
+	}
+	return FString();
+}
+
+TArray<uint8> FOnlineAccountIdRegistryEOSGS::ToReplicationData(const FOnlineAccountIdHandle& Handle) const
+{
+	TArray<uint8> ReplicationData;
+	if (Registry.ValidateOnlineId(Handle))
+	{
+		EOS_ProductUserId ProductUserId = Registry.FindIdValue(Handle);
+		if (ensure(EOS_ProductUserId_IsValid(ProductUserId)))
+		{
+			char EosBuffer[EOS_PRODUCTUSERID_MAX_LENGTH + 1] = {};
+			int32 EosBufferLength = sizeof(EosBuffer);
+			const EOS_EResult EosResult = EOS_ProductUserId_ToString(ProductUserId, EosBuffer, &EosBufferLength);
+			if (ensure(EosResult == EOS_EResult::EOS_Success))
+			{
+				check(EosBufferLength - 1 == OnlineIdEOSUtf8BufferLength);
+				ReplicationData.SetNumUninitialized(OnlineIdEOSHexBufferLength);
+				UE::String::HexToBytes(FUtf8StringView(EosBuffer, OnlineIdEOSUtf8BufferLength), ReplicationData.GetData());
+			}
+		}
+	}
+	return ReplicationData;
+}
+FOnlineAccountIdHandle FOnlineAccountIdRegistryEOSGS::FromReplicationData(const TArray<uint8>& ReplicationData)
+{
+	if (ReplicationData.Num() == OnlineIdEOSHexBufferLength)
+	{
+		char EosBuffer[EOS_EPICACCOUNTID_MAX_LENGTH + 1] = {};
+		UE::String::BytesToHex(TConstArrayView<uint8>(ReplicationData.GetData(), OnlineIdEOSHexBufferLength), EosBuffer);
+		EOS_ProductUserId ProductUserId = EOS_ProductUserId_FromString(EosBuffer);
+		return FindOrAddAccountId(ProductUserId);
+	}
+	return Registry.InvalidHandle;
+}
+
+IOnlineAccountIdRegistryEOSGS& FOnlineAccountIdRegistryEOSGS::GetRegistered()
+{
+	IOnlineAccountIdRegistry* Registry = FOnlineIdRegistryRegistry::Get().GetAccountIdRegistry(EOnlineServices::Epic);
+	check(Registry);
+	return *static_cast<IOnlineAccountIdRegistryEOSGS*>(Registry);
+}
+
+FOnlineAccountIdRegistryEOSGS& FOnlineAccountIdRegistryEOSGS::Get()
+{
+	static FOnlineAccountIdRegistryEOSGS Instance;
+	return Instance;
+}
+
+FOnlineAccountIdHandle FOnlineAccountIdRegistryEOSGS::FindOrAddAccountId(const EOS_ProductUserId ProductUserId)
+{
+	if (ensure(EOS_ProductUserId_IsValid(ProductUserId)))
+	{
+		return Registry.FindOrAddHandle(ProductUserId);
+	}
+	return Registry.InvalidHandle;
+}
+
+EOS_ProductUserId GetProductUserId(const FOnlineAccountIdHandle& Handle)
+{
+	return FOnlineAccountIdRegistryEOSGS::GetRegistered().GetProductUserId(Handle);
+}
+
+EOS_ProductUserId GetProductUserIdChecked(const FOnlineAccountIdHandle& Handle)
+{
+	EOS_ProductUserId ProductUserId = GetProductUserId(Handle);
+	check(EOS_ProductUserId_IsValid(ProductUserId) == EOS_TRUE);
+	return ProductUserId;
+}
+
+FOnlineAccountIdHandle FindAccountId(const EOS_ProductUserId ProductUserId)
+{
+	return FOnlineAccountIdRegistryEOSGS::GetRegistered().FindAccountId(ProductUserId);
+}
+
+FOnlineAccountIdHandle FindAccountIdChecked(const EOS_ProductUserId ProductUserId)
+{
+	FOnlineAccountIdHandle Result = FindAccountId(ProductUserId);
+	check(Result.IsValid());
+	return Result;
+}
+
+/* UE::Online */ }
diff --git a/Engine/Plugins/Online/Experimental/OnlineServicesEOSGS/Source/Private/Online/OnlineServicesEOSGS.cpp b/Engine/Plugins/Online/Experimental/OnlineServicesEOSGS/Source/Private/Online/OnlineServicesEOSGS.cpp
new file mode 100644
--- /dev/null
+++ b/Engine/Plugins/Online/Experimental/OnlineServicesEOSGS/Source/Private/Online/OnlineServicesEOSGS.cpp
@@ -0,0 +1,158 @@
+// Copyright Epic Games, Inc. All Rights Reserved.
+
+#include "Online/OnlineServicesEOSGS.h"
+
+#include "Online/AuthEOSGS.h"
+#include "Online/LobbiesEOSGS.h"
+#include "Online/OnlineIdEOSGS.h"
+#include "Online/OnlineServicesEOSGSTypes.h"
+
+#include "EOSShared.h"
+#include "EOSSharedTypes.h"
+#include "IEOSSDKManager.h"
+
+#if WITH_ENGINE
+#include "InternetAddrEOS.h"
+#include "NetDriverEOSBase.h"
+#include "SocketSubsystemEOSUtils_OnlineServicesEOSGS.h"
+#endif
+
+namespace UE::Online {
+
+struct FEOSPlatformConfig
+{
+	FString ProductId;
+	FString SandboxId;
+	FString DeploymentId;
+	FString ClientId;
+	FString ClientSecret;
+};
+
+namespace Meta {
+
+BEGIN_ONLINE_STRUCT_META(FEOSPlatformConfig)
+	ONLINE_STRUCT_FIELD(FEOSPlatformConfig, ProductId),
+	ONLINE_STRUCT_FIELD(FEOSPlatformConfig, SandboxId),
+	ONLINE_STRUCT_FIELD(FEOSPlatformConfig, DeploymentId),
+	ONLINE_STRUCT_FIELD(FEOSPlatformConfig, ClientId),
+	ONLINE_STRUCT_FIELD(FEOSPlatformConfig, ClientSecret)
+END_ONLINE_STRUCT_META()
+
+/* Meta */ }
+
+FOnlineServicesEOSGS::FOnlineServicesEOSGS(FName InInstanceName)
+	: Super(TEXT("EOS"), InInstanceName)
+{
+}
+
+void FOnlineServicesEOSGS::RegisterComponents()
+{
+	Components.Register<FAuthEOSGS>(*this);
+	Components.Register<FLobbiesEOSGS>(*this);
+	Super::RegisterComponents();
+}
+
+void FOnlineServicesEOSGS::Initialize()
+{
+	IEOSSDKManager* SDKManager = IEOSSDKManager::Get();
+	if (!SDKManager)
+	{
+		return;
+	}
+
+	EOS_EResult InitResult = SDKManager->Initialize();
+	if (InitResult != EOS_EResult::EOS_Success)
+	{
+		return;
+	}
+
+	FEOSPlatformConfig EOSPlatformConfig;
+	LoadConfig(EOSPlatformConfig);
+
+	const FTCHARToUTF8 ProductId(*EOSPlatformConfig.ProductId);
+	const FTCHARToUTF8 SandboxId(*EOSPlatformConfig.SandboxId);
+	const FTCHARToUTF8 DeploymentId(*EOSPlatformConfig.DeploymentId);
+	const FTCHARToUTF8 ClientId(*EOSPlatformConfig.ClientId);
+	const FTCHARToUTF8 ClientSecret(*EOSPlatformConfig.ClientSecret);
+	const FTCHARToUTF8 CacheDirectory(*(SDKManager->GetCacheDirBase() / TEXT("OnlineServicesEOS")));
+
+	EOS_Platform_Options PlatformOptions = {};
+	PlatformOptions.ApiVersion = EOS_PLATFORM_OPTIONS_API_LATEST;
+	PlatformOptions.Reserved = nullptr;
+	PlatformOptions.bIsServer = EOS_FALSE;
+	PlatformOptions.OverrideCountryCode = nullptr;
+	PlatformOptions.OverrideLocaleCode = nullptr;
+	PlatformOptions.Flags = EOS_PF_WINDOWS_ENABLE_OVERLAY_D3D9 | EOS_PF_WINDOWS_ENABLE_OVERLAY_D3D10 | EOS_PF_WINDOWS_ENABLE_OVERLAY_OPENGL; // Enable overlay support for D3D9/10 and OpenGL. This sample uses D3D11 or SDL.
+
+	PlatformOptions.ProductId = ProductId.Get();
+	PlatformOptions.SandboxId = SandboxId.Get();
+	PlatformOptions.DeploymentId = DeploymentId.Get();
+	PlatformOptions.ClientCredentials.ClientId = ClientId.Get();
+	PlatformOptions.ClientCredentials.ClientSecret = ClientSecret.Get();
+	PlatformOptions.CacheDirectory = CacheDirectory.Get();
+
+	EOSPlatformHandle = SDKManager->CreatePlatform(PlatformOptions);
+	if (EOSPlatformHandle)
+	{
+#if WITH_ENGINE
+		SocketSubsystem = MakeShareable(new FSocketSubsystemEOS(EOSPlatformHandle, MakeShareable(new FSocketSubsystemEOSUtils_OnlineServicesEOS(*this))));
+		check(SocketSubsystem);
+
+		FString ErrorStr;
+		if (!SocketSubsystem->Init(ErrorStr))
+		{
+			UE_LOG(LogOnlineServices, Warning, TEXT("[FOnlineServicesEOSGS::Initialize] Unable to initialize Socket Subsystem. Error=[%s]"), *ErrorStr);
+		}
+#endif
+	}
+	else
+	{
+		UE_LOG(LogOnlineServices, Verbose, TEXT("[FOnlineServicesEOSGS::Initialize] Unable to initialize Socket Subsystem. EOS Platform Handle was invalid."));
+		return;
+	}
+
+	Super::Initialize();
+}
+
+TOnlineResult<FGetResolvedConnectString> FOnlineServicesEOSGS::GetResolvedConnectString(FGetResolvedConnectString::Params&& Params)
+{
+	ILobbiesPtr LobbiesEOS = GetLobbiesInterface();
+	check(LobbiesEOS);
+
+	TOnlineResult<FGetJoinedLobbies> JoinedLobbies = LobbiesEOS->GetJoinedLobbies({ Params.LocalUserId });
+	if (JoinedLobbies.IsOk())
+	{
+		for (TSharedRef<const FLobby>& Lobby : JoinedLobbies.GetOkValue().Lobbies)
+		{
+			if (Lobby->LobbyId == Params.LobbyId)
+			{
+#if WITH_ENGINE
+ 				//It should look like this: "EOS:0002aeeb5b2d4388a3752dd6d31222ec:GameNetDriver:97"
+ 				FString NetDriverName = GetDefault<UNetDriverEOSBase>()->NetDriverName.ToString();
+ 				FInternetAddrEOS TempAddr(GetProductUserIdChecked(Lobby->OwnerAccountId), NetDriverName, GetTypeHash(NetDriverName));
+ 				return TOnlineResult<FGetResolvedConnectString>({ TempAddr.ToString(true) });
+#else
+				return TOnlineResult<FGetResolvedConnectString>(Errors::NotImplemented());
+#endif
+			}
+		}
+		// No matching lobby
+		return TOnlineResult<FGetResolvedConnectString>(Errors::InvalidParams());
+	}
+	else
+	{
+		return TOnlineResult<FGetResolvedConnectString>(JoinedLobbies.GetErrorValue());
+	}
+}
+
+EOS_HPlatform FOnlineServicesEOSGS::GetEOSPlatformHandle() const
+{
+	if (EOSPlatformHandle)
+	{
+		return *EOSPlatformHandle;
+	}
+	return nullptr;
+}
+
+
+/* UE::Online */ }
diff --git a/Engine/Plugins/Online/Experimental/OnlineServicesEOSGS/Source/Private/Online/SocketSubsystemEOSUtils_OnlineServicesEOSGS.cpp b/Engine/Plugins/Online/Experimental/OnlineServicesEOSGS/Source/Private/Online/SocketSubsystemEOSUtils_OnlineServicesEOSGS.cpp
new file mode 100644
--- /dev/null
+++ b/Engine/Plugins/Online/Experimental/OnlineServicesEOSGS/Source/Private/Online/SocketSubsystemEOSUtils_OnlineServicesEOSGS.cpp
@@ -0,0 +1,98 @@
+// Copyright Epic Games, Inc. All Rights Reserved.
+
+#if WITH_ENGINE
+#include "SocketSubsystemEOSUtils_OnlineServicesEOSGS.h"
+
+#include "Online/AuthEOSGS.h"
+#include "Online/LobbiesEOSGS.h"
+#include "Online/OnlineIdEOSGS.h"
+#include "Online/OnlineServicesEOSGS.h"
+#include "Online/OnlineServicesEOSGSTypes.h"
+
+namespace UE::Online {
+
+FSocketSubsystemEOSUtils_OnlineServicesEOS::FSocketSubsystemEOSUtils_OnlineServicesEOS(FOnlineServicesEOSGS& InServicesEOSGS)
+	: ServicesEOSGS(InServicesEOSGS)
+{
+}
+
+FSocketSubsystemEOSUtils_OnlineServicesEOS::~FSocketSubsystemEOSUtils_OnlineServicesEOS()
+{
+}
+
+EOS_ProductUserId FSocketSubsystemEOSUtils_OnlineServicesEOS::GetLocalUserId()
+{
+	EOS_ProductUserId Result = nullptr;
+
+	IAuthPtr AuthEOS = ServicesEOSGS.GetAuthInterface();
+	check(AuthEOS);
+
+	FAuthGetAccountByPlatformUserId::Params AuthParams;
+	AuthParams.PlatformUserId = FPlatformMisc::GetPlatformUserForUserIndex(0);
+	TOnlineResult<FAuthGetAccountByPlatformUserId> AuthResult = AuthEOS->GetAccountByPlatformUserId(MoveTemp(AuthParams));
+	if (AuthResult.IsOk())
+	{
+		UE::Online::FAuthGetAccountByPlatformUserId::Result OkValue = AuthResult.GetOkValue();
+
+		Result = GetProductUserIdChecked(OkValue.AccountInfo->UserId);
+	}
+	else
+	{
+		UE_LOG(LogOnlineServices, Verbose, TEXT("[FSocketSubsystemEOSUtils_OnlineServicesEOS::GetLocalUserId] Unable to get account for platform user id [%s]. Error=[%s]"), *ToLogString(AuthParams.PlatformUserId), *AuthResult.GetErrorValue().GetLogString(true));
+	}
+
+	return Result;
+}
+
+FString FSocketSubsystemEOSUtils_OnlineServicesEOS::GetSessionId()
+{
+	FString Result;
+
+	IAuthPtr AuthEOS = ServicesEOSGS.GetAuthInterface();
+	check(AuthEOS);
+
+	FAuthGetAccountByPlatformUserId::Params AuthParams;
+	AuthParams.PlatformUserId = FPlatformMisc::GetPlatformUserForUserIndex(0);
+	TOnlineResult<FAuthGetAccountByPlatformUserId> AuthResult = AuthEOS->GetAccountByPlatformUserId(MoveTemp(AuthParams));
+	if (AuthResult.IsOk())
+	{
+		FAuthGetAccountByPlatformUserId::Result* AuthOkValue = AuthResult.TryGetOkValue();
+
+		ILobbiesPtr LobbiesEOS = ServicesEOSGS.GetLobbiesInterface();
+		check(LobbiesEOS);
+
+		FGetJoinedLobbies::Params LobbiesParams;
+		LobbiesParams.LocalUserId = AuthOkValue->AccountInfo->UserId;
+		TOnlineResult<FGetJoinedLobbies> LobbiesResult = LobbiesEOS->GetJoinedLobbies(MoveTemp(LobbiesParams));
+		if (LobbiesResult.IsOk())
+		{
+			FGetJoinedLobbies::Result* LobbiesOkValue = LobbiesResult.TryGetOkValue();
+
+			// TODO: Pending support in Lobbies interface
+			/*for (TSharedRef<FLobby> Lobby : LobbiesOkValue->Lobbies)
+			{
+					if (Lobby->SessionName == NAME_GameSession)
+					Result = LobbiesEOS->GetLobbyIdString(Lobby);
+			}*/
+		}
+		else
+		{
+			UE_LOG(LogOnlineServices, Verbose, TEXT("[FSocketSubsystemEOSUtils_OnlineServicesEOS::GetSessionId] Unable to get joined lobbies for local user id [%s]. Error=[%s]"), *ToLogString(LobbiesParams.LocalUserId), *AuthResult.GetErrorValue().GetLogString(true));
+		}
+	}
+	else
+	{
+		UE_LOG(LogOnlineServices, Verbose, TEXT("[FSocketSubsystemEOSUtils_OnlineServicesEOS::GetSessionId] Unable to get account for platform user id [%s]. Error=[%s]"), *ToLogString(AuthParams.PlatformUserId), *AuthResult.GetErrorValue().GetLogString(true));
+	}
+
+	return Result;
+}
+
+FName FSocketSubsystemEOSUtils_OnlineServicesEOS::GetSubsystemInstanceName()
+{
+	return ServicesEOSGS.GetInstanceName();
+}
+
+/* UE::Online */}
+
+#endif // WITH_ENGINE
diff --git a/Engine/Plugins/Online/Experimental/OnlineServicesEOSGS/Source/Private/Online/SocketSubsystemEOSUtils_OnlineServicesEOSGS.h b/Engine/Plugins/Online/Experimental/OnlineServicesEOSGS/Source/Private/Online/SocketSubsystemEOSUtils_OnlineServicesEOSGS.h
new file mode 100644
--- /dev/null
+++ b/Engine/Plugins/Online/Experimental/OnlineServicesEOSGS/Source/Private/Online/SocketSubsystemEOSUtils_OnlineServicesEOSGS.h
@@ -0,0 +1,29 @@
+// Copyright Epic Games, Inc. All Rights Reserved.
+
+#pragma once
+
+#if WITH_ENGINE
+
+#include "SocketSubsystemEOS.h"
+
+namespace UE::Online {
+
+class FOnlineServicesEOSGS;
+
+class FSocketSubsystemEOSUtils_OnlineServicesEOS : public ISocketSubsystemEOSUtils
+{
+public:
+	FSocketSubsystemEOSUtils_OnlineServicesEOS(FOnlineServicesEOSGS& InServicesEOS);
+	virtual ~FSocketSubsystemEOSUtils_OnlineServicesEOS();
+
+	virtual EOS_ProductUserId GetLocalUserId() override;
+	virtual FString GetSessionId() override;
+	virtual FName GetSubsystemInstanceName() override;
+
+private:
+	FOnlineServicesEOSGS& ServicesEOSGS;
+};
+
+/* UE::Online */}
+
+#endif // WITH_ENGINE
\ No newline at end of file
diff --git a/Engine/Plugins/Online/Experimental/OnlineServicesEOSGS/Source/Private/OnlineServicesEOSGSModule.cpp b/Engine/Plugins/Online/Experimental/OnlineServicesEOSGS/Source/Private/OnlineServicesEOSGSModule.cpp
new file mode 100644
--- /dev/null
+++ b/Engine/Plugins/Online/Experimental/OnlineServicesEOSGS/Source/Private/OnlineServicesEOSGSModule.cpp
@@ -0,0 +1,40 @@
+// Copyright Epic Games, Inc. All Rights Reserved.
+
+#include "Modules/ModuleInterface.h"
+#include "Modules/ModuleManager.h"
+#include "Online/OnlineIdEOSGS.h"
+#include "Online/OnlineServicesEOSGS.h"
+#include "Online/OnlineServicesRegistry.h"
+
+#include "CoreMinimal.h"
+
+namespace UE::Online
+{
+
+class FOnlineServicesEOSGSModule : public IModuleInterface
+{
+public:
+	virtual void StartupModule() override;
+protected:
+};
+
+class FOnlineServicesFactoryEOSGS : public IOnlineServicesFactory
+{
+public:
+	virtual ~FOnlineServicesFactoryEOSGS() {}
+	virtual TSharedPtr<IOnlineServices> Create(FName InInstanceName) override
+	{
+		return MakeShared<FOnlineServicesEOSGS>(InInstanceName);
+	}
+protected:
+};
+
+void FOnlineServicesEOSGSModule::StartupModule()
+{
+	FOnlineServicesRegistry::Get().RegisterServicesFactory(EOnlineServices::Epic, MakeUnique<FOnlineServicesFactoryEOSGS>());
+	FOnlineIdRegistryRegistry::Get().RegisterAccountIdRegistry(EOnlineServices::Epic, &FOnlineAccountIdRegistryEOSGS::Get());
+}
+
+/* UE::Online */ }
+
+IMPLEMENT_MODULE(UE::Online::FOnlineServicesEOSGSModule, OnlineServicesEOSGS);
diff --git a/Engine/Plugins/Online/Experimental/OnlineServicesEOSGS/Source/Public/Online/AuthEOSGS.h b/Engine/Plugins/Online/Experimental/OnlineServicesEOSGS/Source/Public/Online/AuthEOSGS.h
new file mode 100644
--- /dev/null
+++ b/Engine/Plugins/Online/Experimental/OnlineServicesEOSGS/Source/Public/Online/AuthEOSGS.h
@@ -0,0 +1,98 @@
+// Copyright Epic Games, Inc. All Rights Reserved.
+
+#pragma once
+
+#include "CoreMinimal.h"
+#include "Online/AuthCommon.h"
+
+#if defined(EOS_PLATFORM_BASE_FILE_NAME)
+#include EOS_PLATFORM_BASE_FILE_NAME
+#endif
+#include "eos_auth_types.h"
+#include "eos_connect_types.h"
+
+namespace UE::Online {
+
+const int EOS_STRING_BUFFER_LENGTH = 256;
+const int EOS_MAX_TOKEN_SIZE = 4096;
+
+class FOnlineServicesEOSGS;
+
+/* Base class for EOS_Connect login credentials */
+class ONLINESERVICESEOSGS_API FEOSConnectLoginCredentials : public EOS_Connect_Credentials
+{
+public:
+	FEOSConnectLoginCredentials();
+	virtual ~FEOSConnectLoginCredentials() = default;
+};
+
+/* EOS_Connect login credentials using EAS auth */
+class ONLINESERVICESEOSGS_API FEOSConnectLoginCredentialsEAS
+	: public FEOSConnectLoginCredentials
+	, public FNoncopyable
+{
+public:
+	FEOSConnectLoginCredentialsEAS(EOS_Auth_Token* InEASToken);
+	virtual ~FEOSConnectLoginCredentialsEAS();
+protected:
+	EOS_Auth_Token* EASToken = nullptr;
+};
+
+/* EOS_Auth login credentials */
+struct ONLINESERVICESEOSGS_API FEOSAuthCredentials : public EOS_Auth_Credentials
+{
+	FEOSAuthCredentials();
+	FEOSAuthCredentials(EOS_EExternalCredentialType InExternalType, const TArray<uint8>& InToken);
+
+	FEOSAuthCredentials(const FEOSAuthCredentials& Other);
+	FEOSAuthCredentials& operator=(FEOSAuthCredentials& Other);
+
+	void SetToken(const FCredentialsToken& InToken);
+
+	char IdAnsi[EOS_STRING_BUFFER_LENGTH];
+	char TokenAnsi[EOS_MAX_TOKEN_SIZE];
+};
+
+class ONLINESERVICESEOSGS_API FAuthEOSGS : public FAuthCommon
+{
+public:
+	using Super = FAuthCommon;
+
+	FAuthEOSGS(FOnlineServicesEOSGS& InOwningSubsystem);
+	virtual ~FAuthEOSGS() = default;
+
+	// Begin IAuth
+	virtual void Initialize() override;
+	virtual TOnlineAsyncOpHandle<FAuthLogin> Login(FAuthLogin::Params&& Params) override;
+	virtual TOnlineAsyncOpHandle<FAuthLogout> Logout(FAuthLogout::Params&& Params) override;
+	virtual TOnlineResult<FAuthGetAccountByPlatformUserId> GetAccountByPlatformUserId(FAuthGetAccountByPlatformUserId::Params&& Params) override;
+	virtual TOnlineResult<FAuthGetAccountByAccountId> GetAccountByAccountId(FAuthGetAccountByAccountId::Params&& Params) override;
+	// End IAuth
+
+	// Begin FAuthEOSGS
+	virtual TFuture<FOnlineAccountIdHandle> ResolveAccountId(const FOnlineAccountIdHandle& LocalUserId, const EOS_ProductUserId ProductUserId);
+	virtual TFuture<TArray<FOnlineAccountIdHandle>> ResolveAccountIds(const FOnlineAccountIdHandle& LocalUserId, const TArray<EOS_ProductUserId>& ProductUserIds);
+	virtual TFunction<TFuture<FOnlineAccountIdHandle>(FOnlineAsyncOp& InAsyncOp, const EOS_ProductUserId& ProductUserId)> ResolveProductIdFn();
+	virtual TFunction<TFuture<TArray<FOnlineAccountIdHandle>>(FOnlineAsyncOp& InAsyncOp, const TArray<EOS_ProductUserId>& ProductUserIds)> ResolveProductIdsFn();
+	// End FAuthEOSGS
+
+	bool IsLoggedIn(const FOnlineAccountIdHandle& AccountId) const;
+
+protected:
+	TOnlineChainableAsyncOp<FAuthLogin, TSharedPtr<FEOSConnectLoginCredentials>> LoginEAS(TOnlineAsyncOp<FAuthLogin>& InAsyncOp);
+	void ProcessSuccessfulLogin(TOnlineAsyncOp<FAuthLogin>& InAsyncOp, const EOS_ProductUserId ProductUserId);
+
+	TResult<FOnlineAccountIdHandle, FOnlineError> GetAccountIdByPlatformUserId(FPlatformUserId PlatformUserId) const;
+
+	static FOnlineAccountIdHandle CreateAccountId(const EOS_ProductUserId ProductUserId);
+
+	using FAccountInfoEOS = FAccountInfo;
+	TMap<FOnlineAccountIdHandle, TSharedRef<FAccountInfoEOS>> AccountInfos;
+
+	EOS_HAuth AuthHandle = nullptr;
+	EOS_HConnect ConnectHandle = nullptr;
+};
+
+ELoginStatus ONLINESERVICESEOSGS_API ToELoginStatus(EOS_ELoginStatus InStatus);
+
+/* UE::Online */ }
diff --git a/Engine/Plugins/Online/Experimental/OnlineServicesEOSGS/Source/Public/Online/LobbiesEOSGS.h b/Engine/Plugins/Online/Experimental/OnlineServicesEOSGS/Source/Public/Online/LobbiesEOSGS.h
new file mode 100644
--- /dev/null
+++ b/Engine/Plugins/Online/Experimental/OnlineServicesEOSGS/Source/Public/Online/LobbiesEOSGS.h
@@ -0,0 +1,339 @@
+// Copyright Epic Games, Inc. All Rights Reserved.
+
+#pragma once
+
+#include "CoreMinimal.h"
+#include "Online/LobbiesCommon.h"
+#include "OnlineServicesEOSGSTypes.h"
+
+#if defined(EOS_PLATFORM_BASE_FILE_NAME)
+#include EOS_PLATFORM_BASE_FILE_NAME
+#endif
+#include "eos_lobby_types.h"
+
+namespace UE::Online {
+
+struct FLobbyPrerequisitesEOS;
+class FLobbyDataEOS;
+class FLobbyDataRegistryEOS;
+class FLobbyInviteDataEOS;
+class FOnlineServicesEOSGS;
+class FLobbySearchEOS;
+struct FClientLobbyDataChanges;
+struct FClientLobbyMemberDataChanges;
+	
+class FLobbiesEOSGS : public FLobbiesCommon
+{
+public:
+	FLobbiesEOSGS(FOnlineServicesEOSGS& InServices);
+
+	virtual void Initialize() override;
+	virtual void PreShutdown() override;
+
+	virtual TOnlineAsyncOpHandle<FCreateLobby> CreateLobby(FCreateLobby::Params&& Params) override;
+	virtual TOnlineAsyncOpHandle<FFindLobbies> FindLobbies(FFindLobbies::Params&& Params) override;
+	virtual TOnlineAsyncOpHandle<FJoinLobby> JoinLobby(FJoinLobby::Params&& Params) override;
+	virtual TOnlineAsyncOpHandle<FLeaveLobby> LeaveLobby(FLeaveLobby::Params&& Params) override;
+	virtual TOnlineAsyncOpHandle<FInviteLobbyMember> InviteLobbyMember(FInviteLobbyMember::Params&& Params) override;
+	virtual TOnlineAsyncOpHandle<FDeclineLobbyInvitation> DeclineLobbyInvitation(FDeclineLobbyInvitation::Params&& Params) override;
+	virtual TOnlineAsyncOpHandle<FKickLobbyMember> KickLobbyMember(FKickLobbyMember::Params&& Params) override;
+	virtual TOnlineAsyncOpHandle<FPromoteLobbyMember> PromoteLobbyMember(FPromoteLobbyMember::Params&& Params) override;
+	virtual TOnlineAsyncOpHandle<FModifyLobbyJoinPolicy> ModifyLobbyJoinPolicy(FModifyLobbyJoinPolicy::Params&& Params) override;
+	virtual TOnlineAsyncOpHandle<FModifyLobbyAttributes> ModifyLobbyAttributes(FModifyLobbyAttributes::Params&& Params) override;
+	virtual TOnlineAsyncOpHandle<FModifyLobbyMemberAttributes> ModifyLobbyMemberAttributes(FModifyLobbyMemberAttributes::Params&& Params) override;
+	virtual TOnlineResult<FGetJoinedLobbies> GetJoinedLobbies(FGetJoinedLobbies::Params&& Params) override;
+
+private:
+	void HandleLobbyUpdated(const EOS_Lobby_LobbyUpdateReceivedCallbackInfo* Data);
+	void HandleLobbyMemberUpdated(const EOS_Lobby_LobbyMemberUpdateReceivedCallbackInfo* Data);
+	void HandleLobbyMemberStatusReceived(const EOS_Lobby_LobbyMemberStatusReceivedCallbackInfo* Data);
+	void HandleLobbyInviteReceived(const EOS_Lobby_LobbyInviteReceivedCallbackInfo* Data);
+	void HandleLobbyInviteAccepted(const EOS_Lobby_LobbyInviteAcceptedCallbackInfo* Data);
+	void HandleJoinLobbyAccepted(const EOS_Lobby_JoinLobbyAcceptedCallbackInfo* Data);
+
+protected:
+	void RegisterHandlers();
+	void UnregisterHandlers();
+
+	void AddActiveLobby(FOnlineAccountIdHandle LocalUserId, const TSharedRef<FLobbyDataEOS>& LobbyData);
+	void RemoveActiveLobby(FOnlineAccountIdHandle LocalUserId, const TSharedRef<FLobbyDataEOS>& LobbyData);
+
+	// Todo: store list of invites per lobby.
+	void AddActiveInvite(const TSharedRef<FLobbyInviteDataEOS>& Invite);
+	void RemoveActiveInvite(const TSharedRef<FLobbyInviteDataEOS>& Invite);
+	TSharedPtr<FLobbyInviteDataEOS> GetActiveInvite(FOnlineAccountIdHandle TargetUser, FOnlineLobbyIdHandle TargetLobbyId);
+
+	struct FJoinLobbyImpl
+	{
+		static constexpr TCHAR Name[] = TEXT("JoinLobbyImpl");
+
+		struct Params
+		{
+			// The lobby handle data.
+			TSharedPtr<FLobbyDataEOS> LobbyData;
+
+			// The local user agent which will perform the action.
+			FOnlineAccountIdHandle LocalUserId;
+
+			// The local name for the lobby.
+			FName LocalName;
+
+			// Local users who will be joining the lobby.
+			TArray<FJoinLobbyLocalUserData> LocalUsers;
+		};
+
+		struct Result
+		{
+		};
+	};
+
+	// LobbyData will be fetched from the operation data if not set in Params.
+	TFuture<TDefaultErrorResult<FJoinLobbyImpl>> JoinLobbyImpl(FJoinLobbyImpl::Params&& Params);
+
+	struct FJoinLobbyMemberImpl
+	{
+		static constexpr TCHAR Name[] = TEXT("JoinLobbyMemberImpl");
+
+		struct Params
+		{
+			// The lobby handle data.
+			TSharedPtr<FLobbyDataEOS> LobbyData;
+
+			// The local user agent which will perform the action.
+			FOnlineAccountIdHandle LocalUserId;
+
+			// Initial attributes.
+			TMap<FLobbyAttributeId, FLobbyVariant> Attributes;
+		};
+
+		struct Result
+		{
+		};
+	};
+
+	TOnlineAsyncOpHandle<FJoinLobbyMemberImpl> JoinLobbyMemberImplOp(FJoinLobbyMemberImpl::Params&& Params);
+	TFuture<TDefaultErrorResult<FJoinLobbyMemberImpl>> JoinLobbyMemberImpl(FJoinLobbyMemberImpl::Params&& Params);
+
+	struct FLeaveLobbyImpl
+	{
+		static constexpr TCHAR Name[] = TEXT("LeaveLobbyImpl");
+
+		struct Params
+		{
+			// The lobby handle data.
+			TSharedPtr<FLobbyDataEOS> LobbyData;
+
+			// The local user agent which will perform the action.
+			FOnlineAccountIdHandle LocalUserId;
+		};
+
+		struct Result
+		{
+		};
+	};
+
+	TFuture<TDefaultErrorResult<FLeaveLobbyImpl>> LeaveLobbyImpl(FLeaveLobbyImpl::Params&& Params);
+
+	struct FDestroyLobbyImpl
+	{
+		static constexpr TCHAR Name[] = TEXT("DestroyLobbyImpl");
+
+		struct Params
+		{
+			// The lobby handle data.
+			TSharedPtr<FLobbyDataEOS> LobbyData;
+
+			// The local user agent which will perform the action.
+			FOnlineAccountIdHandle LocalUserId;
+		};
+
+		struct Result
+		{
+		};
+	};
+
+	TFuture<TDefaultErrorResult<FDestroyLobbyImpl>> DestroyLobbyImpl(FDestroyLobbyImpl::Params&& Params);
+
+	struct FInviteLobbyMemberImpl
+	{
+		static constexpr TCHAR Name[] = TEXT("InviteLobbyMemberImpl");
+
+		struct Params
+		{
+			// The lobby handle data.
+			TSharedPtr<FLobbyDataEOS> LobbyData;
+
+			// The local user agent which will perform the action.
+			FOnlineAccountIdHandle LocalUserId;
+
+			// The target user for the invitation.
+			FOnlineAccountIdHandle TargetUserId;
+		};
+
+		struct Result
+		{
+		};
+	};
+
+	TFuture<TDefaultErrorResult<FInviteLobbyMemberImpl>> InviteLobbyMemberImpl(FInviteLobbyMemberImpl::Params&& Params);
+
+	struct FDeclineLobbyInvitationImpl
+	{
+		static constexpr TCHAR Name[] = TEXT("DeclineLobbyInvitationImpl");
+
+		struct Params
+		{
+			// The local user agent which will perform the action.
+			FOnlineAccountIdHandle LocalUserId;
+
+			// Id of the lobby for which the invitations will be declined.
+			FOnlineLobbyIdHandle LobbyId;
+		};
+
+		struct Result
+		{
+		};
+	};
+
+	TFuture<TDefaultErrorResult<FDeclineLobbyInvitationImpl>> DeclineLobbyInvitationImpl(FDeclineLobbyInvitationImpl::Params&& Params);
+
+	struct FKickLobbyMemberImpl
+	{
+		static constexpr TCHAR Name[] = TEXT("KickLobbyMemberImpl");
+
+		struct Params
+		{
+			// The lobby handle data.
+			TSharedPtr<FLobbyDataEOS> LobbyData;
+
+			// The local user agent which will perform the action.
+			FOnlineAccountIdHandle LocalUserId;
+
+			// The target user to be kicked.
+			FOnlineAccountIdHandle TargetUserId;
+		};
+
+		struct Result
+		{
+		};
+	};
+
+	TFuture<TDefaultErrorResult<FKickLobbyMemberImpl>> KickLobbyMemberImpl(FKickLobbyMemberImpl::Params&& Params);
+
+	struct FPromoteLobbyMemberImpl
+	{
+		static constexpr TCHAR Name[] = TEXT("PromoteLobbyMemberImpl");
+
+		struct Params
+		{
+			// The lobby handle data.
+			TSharedPtr<FLobbyDataEOS> LobbyData;
+
+			// The local user agent which will perform the action.
+			FOnlineAccountIdHandle LocalUserId;
+
+			// The target user to be promoted to owner.
+			FOnlineAccountIdHandle TargetUserId;
+		};
+
+		struct Result
+		{
+		};
+	};
+
+	TFuture<TDefaultErrorResult<FPromoteLobbyMemberImpl>> PromoteLobbyMemberImpl(FPromoteLobbyMemberImpl::Params&& Params);
+
+	struct FModifyLobbyDataImpl
+	{
+		static constexpr TCHAR Name[] = TEXT("ModifyLobbyDataImpl");
+
+		struct Params
+		{
+			// The lobby handle data.
+			TSharedPtr<FLobbyDataEOS> LobbyData;
+
+			// The local user agent which will perform the action.
+			FOnlineAccountIdHandle LocalUserId;
+
+			// The changes to apply to the lobby data.
+			TSharedPtr<FClientLobbyDataChanges> Changes;
+		};
+
+		// Todo: += operator.
+		// Mergeable op must be queued by lobby id.
+
+		struct Result
+		{
+		};
+	};
+
+	TFuture<TDefaultErrorResult<FModifyLobbyDataImpl>> ModifyLobbyDataImpl(FModifyLobbyDataImpl::Params&& Params);
+
+	struct FModifyLobbyMemberDataImpl
+	{
+		static constexpr TCHAR Name[] = TEXT("ModifyLobbyMemberDataImpl");
+
+		struct Params
+		{
+			// The lobby handle data.
+			TSharedPtr<FLobbyDataEOS> LobbyData;
+
+			// The local user agent which will perform the action.
+			FOnlineAccountIdHandle LocalUserId;
+
+			// The changes to apply to the lobby member data.
+			TSharedPtr<FClientLobbyMemberDataChanges> Changes;
+		};
+
+		// Todo: += operator.
+		// Mergeable op must be queued on union of lobby id + lobby member account id.
+
+		struct Result
+		{
+		};
+	};
+
+	TFuture<TDefaultErrorResult<FModifyLobbyMemberDataImpl>> ModifyLobbyMemberDataImpl(FModifyLobbyMemberDataImpl::Params&& Params);
+
+	struct FProcessLobbyNotificationImpl
+	{
+		static constexpr TCHAR Name[] = TEXT("ProcessLobbyNotificationImpl");
+
+		struct Params
+		{
+			// The lobby handle data.
+			TSharedPtr<FLobbyDataEOS> LobbyData;
+
+			// Joining / mutated members.
+			TSet<EOS_ProductUserId> MutatedMembers;
+
+			// Leaving members.
+			TMap<EOS_ProductUserId, ELobbyMemberLeaveReason> LeavingMembers;
+		};
+
+		// Todo: += operator.
+		// Mergeable op must be queued by lobby id.
+
+		struct Result
+		{
+		};
+	};
+
+	TOnlineAsyncOpHandle<FProcessLobbyNotificationImpl> ProcessLobbyNotificationImplOp(FProcessLobbyNotificationImpl::Params&& Params);
+
+	EOSEventRegistrationPtr OnLobbyUpdatedEOSEventRegistration;
+	EOSEventRegistrationPtr OnLobbyMemberUpdatedEOSEventRegistration;
+	EOSEventRegistrationPtr OnLobbyMemberStatusReceivedEOSEventRegistration;
+	EOSEventRegistrationPtr OnLobbyInviteReceivedEOSEventRegistration;
+	EOSEventRegistrationPtr OnLobbyInviteAcceptedEOSEventRegistration;
+	EOSEventRegistrationPtr OnJoinLobbyAcceptedEOSEventRegistration;
+
+	TSharedPtr<FLobbyPrerequisitesEOS> LobbyPrerequisites;
+	TSharedPtr<FLobbyDataRegistryEOS> LobbyDataRegistry;
+
+	TMap<FOnlineAccountIdHandle, TSet<TSharedRef<FLobbyDataEOS>>> ActiveLobbies;
+	TMap<FOnlineAccountIdHandle, TMap<FOnlineLobbyIdHandle, TSharedRef<FLobbyInviteDataEOS>>> ActiveInvites;
+	TMap<FOnlineAccountIdHandle, TSharedRef<FLobbySearchEOS>> ActiveSearchResults;
+};
+
+/* UE::Online */ }
diff --git a/Engine/Plugins/Online/Experimental/OnlineServicesEOSGS/Source/Public/Online/OnlineIdEOSGS.h b/Engine/Plugins/Online/Experimental/OnlineServicesEOSGS/Source/Public/Online/OnlineIdEOSGS.h
new file mode 100644
--- /dev/null
+++ b/Engine/Plugins/Online/Experimental/OnlineServicesEOSGS/Source/Public/Online/OnlineIdEOSGS.h
@@ -0,0 +1,65 @@
+// Copyright Epic Games, Inc. All Rights Reserved.
+
+#pragma once
+
+#include "Online/CoreOnline.h"
+#include "Online/OnlineIdCommon.h"
+
+#if defined(EOS_PLATFORM_BASE_FILE_NAME)
+#include EOS_PLATFORM_BASE_FILE_NAME
+#endif
+#include "eos_common.h"
+
+namespace UE::Online {
+
+class IOnlineAccountIdRegistryEOSGS : public IOnlineAccountIdRegistry
+{
+public:
+	virtual FOnlineAccountIdHandle FindAccountId(EOS_ProductUserId ProductUserId) const = 0;
+	virtual EOS_ProductUserId GetProductUserId(const FOnlineAccountIdHandle& Handle) const = 0;
+};
+
+/**
+ * Account id registry specifically for EOS id's which are segmented.
+ */
+class FOnlineAccountIdRegistryEOSGS
+	: public IOnlineAccountIdRegistryEOSGS
+{
+public:
+	virtual ~FOnlineAccountIdRegistryEOSGS() = default;
+
+	// Begin IOnlineAccountIdRegistryEOSGS
+	virtual FOnlineAccountIdHandle FindAccountId(const EOS_ProductUserId ProductUserId) const override;
+	virtual EOS_ProductUserId GetProductUserId(const FOnlineAccountIdHandle& Handle) const override;
+	// End IOnlineAccountIdRegistryEOSGS
+
+	// Begin IOnlineAccountIdRegistry
+	virtual FString ToLogString(const FOnlineAccountIdHandle& Handle) const override;
+	virtual TArray<uint8> ToReplicationData(const FOnlineAccountIdHandle& Handle) const override;
+	virtual FOnlineAccountIdHandle FromReplicationData(const TArray<uint8>& ReplicationData) override;
+	// End IOnlineAccountIdRegistry
+
+	static IOnlineAccountIdRegistryEOSGS& GetRegistered();
+
+private:
+	// FAuthEOSGS is the only thing that should be able to create PUID-only net ids in this registry, in its resolve methods.
+	friend class FAuthEOSGS;
+	friend class FOnlineServicesEOSGSModule;
+	static FOnlineAccountIdRegistryEOSGS& Get();
+	FOnlineAccountIdHandle FindOrAddAccountId(const EOS_ProductUserId ProductUserId);
+
+	TOnlineBasicAccountIdRegistry<EOS_ProductUserId, EOnlineServices::Epic> Registry;
+};
+
+EOS_ProductUserId ONLINESERVICESEOSGS_API GetProductUserId(const FOnlineAccountIdHandle& Handle);
+EOS_ProductUserId ONLINESERVICESEOSGS_API GetProductUserIdChecked(const FOnlineAccountIdHandle& Handle);
+FOnlineAccountIdHandle ONLINESERVICESEOSGS_API FindAccountId(const EOS_ProductUserId EpicAccountId);
+FOnlineAccountIdHandle ONLINESERVICESEOSGS_API FindAccountIdChecked(const EOS_ProductUserId EpicAccountId);
+
+template<typename IdType>
+inline bool ValidateOnlineId(const TOnlineIdHandle<IdType> Handle)
+{
+	return Handle.GetOnlineServicesType() == EOnlineServices::Epic && Handle.IsValid();
+}
+
+} /* namespace UE::Online */
diff --git a/Engine/Plugins/Online/Experimental/OnlineServicesEOSGS/Source/Public/Online/OnlineServicesEOSGS.h b/Engine/Plugins/Online/Experimental/OnlineServicesEOSGS/Source/Public/Online/OnlineServicesEOSGS.h
new file mode 100644
--- /dev/null
+++ b/Engine/Plugins/Online/Experimental/OnlineServicesEOSGS/Source/Public/Online/OnlineServicesEOSGS.h
@@ -0,0 +1,42 @@
+// Copyright Epic Games, Inc. All Rights Reserved.
+
+#pragma once
+
+#include "Online/OnlineServicesCommon.h"
+
+#if defined(EOS_PLATFORM_BASE_FILE_NAME)
+#include EOS_PLATFORM_BASE_FILE_NAME
+#endif
+#include "eos_sdk.h"
+
+#if WITH_ENGINE
+class FSocketSubsystemEOS;
+#endif
+
+using IEOSPlatformHandlePtr = TSharedPtr<class IEOSPlatformHandle>;
+
+namespace UE::Online {
+
+class ONLINESERVICESEOSGS_API FOnlineServicesEOSGS : public FOnlineServicesCommon
+{
+public:
+	using Super = FOnlineServicesCommon;
+
+	FOnlineServicesEOSGS(FName InInstanceName);
+	virtual ~FOnlineServicesEOSGS() = default;
+
+	virtual void RegisterComponents() override;
+	virtual void Initialize() override;
+	virtual TOnlineResult<FGetResolvedConnectString> GetResolvedConnectString(FGetResolvedConnectString::Params&& Params) override;
+	virtual EOnlineServices GetServicesProvider() const override { return EOnlineServices::Epic; }
+
+	EOS_HPlatform GetEOSPlatformHandle() const;
+protected:
+	IEOSPlatformHandlePtr EOSPlatformHandle;
+
+#if WITH_ENGINE
+	TSharedPtr<FSocketSubsystemEOS, ESPMode::ThreadSafe> SocketSubsystem;
+#endif
+};
+
+/* UE::Online */ }
diff --git a/Engine/Plugins/Online/Experimental/OnlineServicesEOSGS/Source/Public/Online/OnlineServicesEOSGSTypes.h b/Engine/Plugins/Online/Experimental/OnlineServicesEOSGS/Source/Public/Online/OnlineServicesEOSGSTypes.h
new file mode 100644
--- /dev/null
+++ b/Engine/Plugins/Online/Experimental/OnlineServicesEOSGS/Source/Public/Online/OnlineServicesEOSGSTypes.h
@@ -0,0 +1,219 @@
+// Copyright Epic Games, Inc. All Rights Reserved.
+
+#pragma once
+
+#include "Online/OnlineAsyncOp.h"
+#include "Online/OnlineError.h"
+#include "Async/Future.h"
+
+#include "EOSShared.h"
+
+namespace UE::Online {
+
+inline FOnlineError FromEOSError(EOS_EResult ResultCode)
+{
+	// Todo: make this for real.
+	return FOnlineError(Errors::Unknown());
+}
+
+/** Class to handle all callbacks generically using a future to forward callback results */
+template<typename CallbackType>
+class TEOSCallback
+{
+	using CallbackFuncType = void (EOS_CALL*)(const CallbackType*);
+public:
+	TEOSCallback() = default;
+	TEOSCallback(TPromise<const CallbackType*>&& InPromise) : Promise(MoveTemp(InPromise)) {}
+	virtual ~TEOSCallback() = default;
+
+	operator CallbackFuncType()
+	{
+		return &CallbackImpl;
+	}
+
+	TFuture<const CallbackType*> GetFuture() { return Promise.GetFuture(); }
+
+private:
+	TPromise<const CallbackType*> Promise;
+
+	static void EOS_CALL CallbackImpl(const CallbackType* Data)
+	{
+		if (EOS_EResult_IsOperationComplete(Data->ResultCode) == EOS_FALSE)
+		{
+			// Ignore
+			return;
+		}
+		check(IsInGameThread());
+
+		TEOSCallback* CallbackThis = reinterpret_cast<TEOSCallback*>(Data->ClientData);
+		check(CallbackThis);
+
+		CallbackThis->Promise.EmplaceValue(Data);
+
+		delete CallbackThis;
+	}
+};
+
+template<typename TEOSResult, typename TEOSHandle, typename TEOSParameters, typename TEOSFn> 
+TFuture<const TEOSResult*> EOS_Async(TEOSFn EOSFn, TEOSHandle EOSHandle, TEOSParameters Parameters)
+{
+	TEOSCallback<TEOSResult>* Callback = new TEOSCallback<TEOSResult>();
+	EOSFn(EOSHandle, &Parameters, Callback, *Callback);
+	return Callback->GetFuture();
+}
+
+template<typename TEOSResult, typename TEOSHandle, typename TEOSParameters, typename TEOSFn>
+void EOS_Async(TEOSFn EOSFn, TEOSHandle EOSHandle, TEOSParameters Parameters, TPromise<const TEOSResult*>&& Promise)
+{
+	TEOSCallback<TEOSResult>* Callback = new TEOSCallback<TEOSResult>(MoveTemp(Promise));
+	EOSFn(EOSHandle, &Parameters, Callback, *Callback);
+}
+
+class EOSEventRegistration
+{
+public:
+	virtual ~EOSEventRegistration() = default;
+};
+typedef TUniquePtr<EOSEventRegistration> EOSEventRegistrationPtr;
+
+/** 
+* EOS event registration utility for binding an EOS notifier registration to a RAII object which handles
+* unregistering when it exits scope. Intended to be used from a TOnlineComponent class.
+* 
+* Example:
+*	EOSEventRegistrationPtr OnLobbyUpdatedEOSEventRegistration = EOS_RegisterComponentEventHandler(
+*		this,
+*		LobbyHandle,
+*		EOS_LOBBY_ADDNOTIFYLOBBYUPDATERECEIVED_API_LATEST,
+*		&EOS_Lobby_AddNotifyLobbyUpdateReceived,
+*		&EOS_Lobby_RemoveNotifyLobbyUpdateReceived,
+*		&FLobbiesEOS::HandleLobbyUpdated);
+**/
+template <
+	typename ComponentHandlerClass,
+	typename EOSHandle,
+	typename EOSNotfyRegisterFunction,
+	typename EOSNotfyUnregisterFunction,
+	typename ComponentHandlerFunction>
+EOSEventRegistrationPtr EOS_RegisterComponentEventHandler(
+	ComponentHandlerClass* HandlerClass,
+	EOSHandle ClientHandle,
+	int32_t ApiVersion,
+	EOSNotfyRegisterFunction NotfyRegisterFunction,
+	EOSNotfyUnregisterFunction NotfyUnregisterFunction,
+	ComponentHandlerFunction HandlerFunction);
+
+namespace detail {
+
+template<typename Function> struct TEOSCallbackTraitsBase;
+
+template<typename Component, typename EventData>
+struct TEOSCallbackTraitsBase<void (Component::*)(const EventData*)>
+{
+	typedef Component ComponentType;
+	typedef EventData EventDataType;
+};
+
+template<typename Function>
+struct TEOSCallbackTraits : public TEOSCallbackTraitsBase<typename std::remove_reference<Function>::type>
+{
+};
+
+template<typename Function> struct TEOSNotifyRegisterTraits;
+
+template<typename EOSHandle, typename Options, typename ClientData, typename NotificationFn, typename NotificationId>
+struct TEOSNotifyRegisterTraits<NotificationId (*)(EOSHandle, const Options*, ClientData*, NotificationFn)>
+{
+	typedef Options OptionsType;
+	typedef NotificationId NotificationIdType;
+};
+
+template <
+	typename ComponentHandlerClass,
+	typename EOSHandle,
+	typename EOSNotfyRegisterFunction,
+	typename EOSNotfyUnregisterFunction,
+	typename ComponentHandlerFunction>
+class EOSEventRegistrationImpl : public EOSEventRegistration
+{
+public:
+	typedef EOSEventRegistrationImpl<
+		ComponentHandlerClass,
+		EOSHandle,
+		EOSNotfyRegisterFunction,
+		EOSNotfyUnregisterFunction,
+		ComponentHandlerFunction> ThisClass;
+
+	EOSEventRegistrationImpl(
+		ComponentHandlerClass* HandlerClass,
+		EOSHandle ClientHandle,
+		int32_t ApiVersion,
+		EOSNotfyRegisterFunction NotfyRegisterFunction,
+		EOSNotfyUnregisterFunction NotfyUnregisterFunction,
+		ComponentHandlerFunction HandlerFunction)
+		: HandlerClass(HandlerClass)
+		, ClientHandle(ClientHandle)
+		, NotfyUnregisterFunction(NotfyUnregisterFunction)
+		, HandlerFunction(HandlerFunction)
+	{
+		typename TEOSNotifyRegisterTraits<EOSNotfyRegisterFunction>::OptionsType Options = { };
+		Options.ApiVersion = ApiVersion;
+		NotificationId = NotfyRegisterFunction(ClientHandle, &Options, this,
+		[](const typename TEOSCallbackTraits<ComponentHandlerFunction>::EventDataType* Data)
+		{
+			ThisClass* This = reinterpret_cast<ThisClass*>(Data->ClientData);
+			(This->HandlerClass->*This->HandlerFunction)(Data);
+		});
+	}
+
+	virtual ~EOSEventRegistrationImpl()
+	{
+		NotfyUnregisterFunction(ClientHandle, NotificationId);
+	}
+
+private:
+	EOSEventRegistrationImpl() = delete;
+	EOSEventRegistrationImpl(const EOSEventRegistrationImpl&) = delete;
+	EOSEventRegistrationImpl& operator=(const EOSEventRegistrationImpl&) = delete;
+
+	typedef typename TEOSNotifyRegisterTraits<EOSNotfyRegisterFunction>::NotificationIdType NotificationIdType;
+
+	NotificationIdType NotificationId;
+	ComponentHandlerClass* HandlerClass;
+	EOSHandle ClientHandle;
+	EOSNotfyUnregisterFunction NotfyUnregisterFunction;
+	ComponentHandlerFunction HandlerFunction;
+};
+
+/* detail */ }
+
+template <
+	typename ComponentHandlerClass,
+	typename EOSHandle,
+	typename EOSNotfyRegisterFunction,
+	typename EOSNotfyUnregisterFunction,
+	typename ComponentHandlerFunction>
+EOSEventRegistrationPtr EOS_RegisterComponentEventHandler(
+	ComponentHandlerClass* HandlerClass,
+	EOSHandle ClientHandle,
+	int32_t ApiVersion,
+	EOSNotfyRegisterFunction NotfyRegisterFunction,
+	EOSNotfyUnregisterFunction NotfyUnregisterFunction,
+	ComponentHandlerFunction HandlerFunction)
+{
+	return MakeUnique<
+		detail::EOSEventRegistrationImpl<
+			ComponentHandlerClass,
+			EOSHandle,
+			EOSNotfyRegisterFunction,
+			EOSNotfyUnregisterFunction,
+			ComponentHandlerFunction>>(
+			HandlerClass,
+			ClientHandle,
+			ApiVersion,
+			NotfyRegisterFunction,
+			NotfyUnregisterFunction,
+			HandlerFunction);
+}
+
+/* UE::Online */ }
