- Fix broken build

diff --git a/Engine/Plugins/Online/Experimental/OnlineServicesEOSGS/Source/Private/Online/LobbiesEOSGSTypes.h b/Engine/Plugins/Online/Experimental/OnlineServicesEOSGS/Source/Private/Online/LobbiesEOSGSTypes.h
--- a/Engine/Plugins/Online/Experimental/OnlineServicesEOSGS/Source/Private/Online/LobbiesEOSGSTypes.h
+++ b/Engine/Plugins/Online/Experimental/OnlineServicesEOSGS/Source/Private/Online/LobbiesEOSGSTypes.h
@@ -22,7 +22,7 @@ class FLobbyDetailsInfoEOS;
 class FLobbyInviteDataEOSRegistry;
 class FLobbySchema;
 struct FLobbyEvents;
-class FAuthEOS;
+class FAuthEOSGS;
 struct FClientLobbyMemberSnapshot;
 
 //--------------------------------------------------------------------------------------------------
