- Achievements EOS implementation

diff --git a/Engine/Plugins/Online/Experimental/OnlineServicesEOSGS/Source/Private/Online/AchievementsEOSGS.cpp b/Engine/Plugins/Online/Experimental/OnlineServicesEOSGS/Source/Private/Online/AchievementsEOSGS.cpp
new file mode 100644
--- /dev/null
+++ b/Engine/Plugins/Online/Experimental/OnlineServicesEOSGS/Source/Private/Online/AchievementsEOSGS.cpp
@@ -0,0 +1,397 @@
+// Copyright Epic Games, Inc. All Rights Reserved.
+
+#include "Online/AchievementsEOSGS.h"
+
+#include "Online/AchievementsErrors.h"
+#include "Online/AuthEOSGS.h"
+#include "Online/OnlineIdEOSGS.h"
+#include "Online/OnlineServicesEOSGS.h"
+#include "Online/OnlineServicesEOSGSTypes.h"
+
+#include "eos_achievements.h"
+#include "eos_achievements_types.h"
+
+namespace UE::Online {
+
+FAchievementsEOSGS::FAchievementsEOSGS(FOnlineServicesEOSGS& InServices)
+	: Super(InServices)
+{
+}
+
+void FAchievementsEOSGS::Initialize()
+{
+	Super::Initialize();
+
+	AchievementsHandle = EOS_Platform_GetAchievementsInterface(static_cast<FOnlineServicesEOSGS&>(GetServices()).GetEOSPlatformHandle());
+	check(AchievementsHandle);
+
+	// Register for achievement unlocked events
+	EOS_Achievements_AddNotifyAchievementsUnlockedV2Options Options = { };
+	Options.ApiVersion = EOS_ACHIEVEMENTS_ADDNOTIFYACHIEVEMENTSUNLOCKEDV2_API_LATEST;
+	static_assert(EOS_ACHIEVEMENTS_ADDNOTIFYACHIEVEMENTSUNLOCKEDV2_API_LATEST == 2, "EOS_Achievements_AddNotifyAchievementsUnlockedV2Options updated, check new fields");
+	NotifyAchievementsUnlockedNotificationId = EOS_Achievements_AddNotifyAchievementsUnlockedV2(AchievementsHandle, &Options, this, [](const EOS_Achievements_OnAchievementsUnlockedCallbackV2Info* Data)
+	{
+		FAchievementsEOSGS* This = reinterpret_cast<FAchievementsEOSGS*>(Data->ClientData);
+
+		FOnlineAccountIdHandle LocalUserId = FindAccountId(Data->UserId);
+		if (LocalUserId.IsValid())
+		{
+			FAchievementStateUpdated EventParams;
+			EventParams.LocalUserId = LocalUserId;
+			EventParams.AchievementIds.Emplace(UTF8_TO_TCHAR(Data->AchievementId));
+
+			This->OnAchievementStateUpdatedEvent.Broadcast(EventParams);
+		}
+	});
+}
+
+void FAchievementsEOSGS::Shutdown()
+{
+	Super::Shutdown();
+
+	if (NotifyAchievementsUnlockedNotificationId != EOS_INVALID_NOTIFICATIONID)
+	{
+		EOS_Achievements_RemoveNotifyAchievementsUnlocked(AchievementsHandle, NotifyAchievementsUnlockedNotificationId);
+		NotifyAchievementsUnlockedNotificationId = EOS_INVALID_NOTIFICATIONID;
+	}
+}
+
+TOnlineAsyncOpHandle<FQueryAchievementDefinitions> FAchievementsEOSGS::QueryAchievementDefinitions(FQueryAchievementDefinitions::Params&& Params)
+{
+	TOnlineAsyncOpRef<FQueryAchievementDefinitions> Op = GetOp<FQueryAchievementDefinitions>(MoveTemp(Params));
+
+	if (!Services.Get<FAuthEOSGS>()->IsLoggedIn(Op->GetParams().LocalUserId))
+	{
+		Op->SetError(Errors::InvalidUser());
+		return Op->GetHandle();
+	}
+
+	if (!AchievementDefinitions.IsSet())
+	{
+		Op->Then([this](TOnlineAsyncOp<FQueryAchievementDefinitions>& InAsyncOp)
+		{
+			EOS_Achievements_QueryDefinitionsOptions Options = {};
+			Options.ApiVersion = EOS_ACHIEVEMENTS_QUERYDEFINITIONS_API_LATEST;
+			static_assert(EOS_ACHIEVEMENTS_QUERYDEFINITIONS_API_LATEST == 3, "EOS_Achievements_QueryDefinitionsOptions updated, check new fields");
+			Options.LocalUserId = GetProductUserIdChecked(InAsyncOp.GetParams().LocalUserId);
+
+			return EOS_Async<EOS_Achievements_OnQueryDefinitionsCompleteCallbackInfo>(EOS_Achievements_QueryDefinitions, AchievementsHandle, Options);
+		})
+		.Then([this](TOnlineAsyncOp<FQueryAchievementDefinitions>& InAsyncOp, const EOS_Achievements_OnQueryDefinitionsCompleteCallbackInfo* Data)
+		{
+			if (Data->ResultCode != EOS_EResult::EOS_Success)
+			{
+				UE_LOG(LogTemp, Warning, TEXT("EOS_Achievements_QueryDefinitions failed with result=[%s]"), *LexToString(Data->ResultCode));
+
+				InAsyncOp.SetError(FromEOSError(Data->ResultCode));
+				return;
+			}
+
+			EOS_Achievements_GetAchievementDefinitionCountOptions GetCountOptions = {};
+			GetCountOptions.ApiVersion = EOS_ACHIEVEMENTS_GETACHIEVEMENTDEFINITIONCOUNT_API_LATEST;
+			static_assert(EOS_ACHIEVEMENTS_GETACHIEVEMENTDEFINITIONCOUNT_API_LATEST == 1, "EOS_Achievements_GetAchievementDefinitionCountOptions updated, check new fields");
+
+			const uint32_t AchievementCount = EOS_Achievements_GetAchievementDefinitionCount(AchievementsHandle, &GetCountOptions);
+
+			FAchievementDefinitionMap NewAchievementDefinitions;
+			for (uint32_t AchievementIdx = 0; AchievementIdx < AchievementCount; AchievementIdx++)
+			{
+				EOS_Achievements_CopyAchievementDefinitionV2ByIndexOptions CopyOptions = {};
+				CopyOptions.ApiVersion = EOS_ACHIEVEMENTS_COPYACHIEVEMENTDEFINITIONV2BYINDEX_API_LATEST;
+				static_assert(EOS_ACHIEVEMENTS_COPYACHIEVEMENTDEFINITIONV2BYINDEX_API_LATEST == 2, "EOS_Achievements_CopyAchievementDefinitionV2ByIndexOptions updated, check new fields");
+				CopyOptions.AchievementIndex = AchievementIdx;
+
+				EOS_Achievements_DefinitionV2* EOSDefinition = nullptr;
+				const EOS_EResult CopyResult = EOS_Achievements_CopyAchievementDefinitionV2ByIndex(AchievementsHandle, &CopyOptions, &EOSDefinition);
+				if (CopyResult != EOS_EResult::EOS_Success)
+				{
+					UE_LOG(LogTemp, Warning, TEXT("EOS_Achievements_CopyAchievementDefinitionV2ByIndex failed with result=[%s]"), *LexToString(CopyResult));
+
+					InAsyncOp.SetError(FromEOSError(CopyResult));
+					return;
+				}
+
+				static_assert(EOS_ACHIEVEMENTS_DEFINITIONV2_API_LATEST == 2, "EOS_Achievements_DefinitionV2 updated, check new fields");
+				if (!ensure(EOSDefinition->ApiVersion == EOS_ACHIEVEMENTS_DEFINITIONV2_API_LATEST))
+				{
+					UE_LOG(LogTemp, Warning, TEXT("EOS_Achievements_DefinitionV2 version mismatch Expected=%d Actual=%d"), EOS_ACHIEVEMENTS_DEFINITIONV2_API_LATEST, EOSDefinition->ApiVersion);
+				}
+
+				FString AchievementId = UTF8_TO_TCHAR(EOSDefinition->AchievementId);
+				FAchievementDefinition& AchievementDefinition = NewAchievementDefinitions.Emplace(AchievementId);
+				AchievementDefinition.AchievementId = MoveTemp(AchievementId);
+				AchievementDefinition.UnlockedDisplayName = FText::FromString(UTF8_TO_TCHAR(EOSDefinition->UnlockedDisplayName));
+				AchievementDefinition.UnlockedDescription = FText::FromString(UTF8_TO_TCHAR(EOSDefinition->UnlockedDescription));
+				AchievementDefinition.LockedDisplayName = FText::FromString(UTF8_TO_TCHAR(EOSDefinition->LockedDisplayName));
+				AchievementDefinition.LockedDescription = FText::FromString(UTF8_TO_TCHAR(EOSDefinition->LockedDescription));
+				AchievementDefinition.FlavorText = FText::FromString(UTF8_TO_TCHAR(EOSDefinition->FlavorText));
+				AchievementDefinition.UnlockedIconUrl = UTF8_TO_TCHAR(EOSDefinition->UnlockedIconURL);
+				AchievementDefinition.LockedIconUrl = UTF8_TO_TCHAR(EOSDefinition->LockedIconURL);
+				AchievementDefinition.bIsHidden = EOSDefinition->bIsHidden == EOS_TRUE;
+				for (uint32_t StatThresholdsIdx = 0; StatThresholdsIdx < EOSDefinition->StatThresholdsCount; StatThresholdsIdx++)
+				{
+					static_assert(EOS_ACHIEVEMENTS_STATTHRESHOLDS_API_LATEST == 1, "EOS_Achievements_StatThresholds updated, check new fields");
+					if (!ensure(EOSDefinition->ApiVersion == EOS_ACHIEVEMENTS_STATTHRESHOLDS_API_LATEST))
+					{
+						UE_LOG(LogTemp, Warning, TEXT("EOS_Achievements_StatThresholds version mismatch Expected=%d Actual=%d"), EOS_ACHIEVEMENTS_DEFINITIONV2_API_LATEST, EOSDefinition->ApiVersion);
+					}
+
+					const EOS_Achievements_StatThresholds& EOSStatThreshold = EOSDefinition->StatThresholds[StatThresholdsIdx];
+
+					FAchievementStatDefinition& StatDefinition = AchievementDefinition.StatDefinitions.Emplace_GetRef();
+					StatDefinition.StatId = UTF8_TO_TCHAR(EOSStatThreshold.Name);
+					StatDefinition.UnlockThreshold = uint32_t(EOSStatThreshold.Threshold);
+				}
+				EOS_Achievements_DefinitionV2_Release(EOSDefinition);
+			}
+
+			AchievementDefinitions.Emplace(MoveTemp(NewAchievementDefinitions));
+			InAsyncOp.SetResult({});
+		})
+		.Enqueue(GetSerialQueue());
+	}
+	else
+	{
+		Op->SetResult({});
+	}
+	return Op->GetHandle();
+}
+
+TOnlineResult<FGetAchievementIds> FAchievementsEOSGS::GetAchievementIds(FGetAchievementIds::Params&& Params)
+{
+	if (!Services.Get<FAuthEOSGS>()->IsLoggedIn(Params.LocalUserId))
+	{
+		return TOnlineResult<FGetAchievementIds>(Errors::InvalidUser());
+	}
+
+	if (!AchievementDefinitions.IsSet())
+	{
+		// Call QueryAchievementDefinitions first
+		return TOnlineResult<FGetAchievementIds>(Errors::InvalidState());
+	}
+
+	FGetAchievementIds::Result Result;
+	AchievementDefinitions->GenerateKeyArray(Result.AchievementIds);
+	return TOnlineResult<FGetAchievementIds>(MoveTemp(Result));
+}
+
+TOnlineResult<FGetAchievementDefinition> FAchievementsEOSGS::GetAchievementDefinition(FGetAchievementDefinition::Params&& Params)
+{
+	if (!Services.Get<FAuthEOSGS>()->IsLoggedIn(Params.LocalUserId))
+	{
+		return TOnlineResult<FGetAchievementDefinition>(Errors::InvalidUser());
+	}
+
+	if (!AchievementDefinitions.IsSet())
+	{
+		// Should call QueryAchievementDefinitions first
+		return TOnlineResult<FGetAchievementDefinition>(Errors::InvalidState());
+	}
+
+	const FAchievementDefinition* AchievementDefinition = AchievementDefinitions->Find(Params.AchievementId);
+	if (!AchievementDefinition)
+	{
+		return TOnlineResult<FGetAchievementDefinition>(Errors::NotFound());
+	}
+
+	return TOnlineResult<FGetAchievementDefinition>({ *AchievementDefinition });
+}
+
+TOnlineAsyncOpHandle<FQueryAchievementStates> FAchievementsEOSGS::QueryAchievementStates(FQueryAchievementStates::Params&& Params)
+{
+	TOnlineAsyncOpRef<FQueryAchievementStates> Op = GetOp<FQueryAchievementStates>(MoveTemp(Params));
+
+	if (!Services.Get<FAuthEOSGS>()->IsLoggedIn(Op->GetParams().LocalUserId))
+	{
+		Op->SetError(Errors::InvalidUser());
+		return Op->GetHandle();
+	}
+
+	if (!AchievementDefinitions.IsSet())
+	{
+		// Call QueryAchievementDefinitions first
+		Op->SetError(Errors::InvalidState());
+		return Op->GetHandle();
+	}
+
+	if (!AchievementStates.Contains(Op->GetParams().LocalUserId))
+	{
+		Op->Then([this](TOnlineAsyncOp<FQueryAchievementStates>& InAsyncOp)
+		{
+			EOS_Achievements_QueryPlayerAchievementsOptions Options = {};
+			Options.ApiVersion = EOS_ACHIEVEMENTS_QUERYPLAYERACHIEVEMENTS_API_LATEST;
+			static_assert(EOS_ACHIEVEMENTS_QUERYPLAYERACHIEVEMENTS_API_LATEST == 2, "EOS_Achievements_QueryPlayerAchievementsOptions updated, check new fields");
+			Options.LocalUserId = GetProductUserIdChecked(InAsyncOp.GetParams().LocalUserId);
+			Options.TargetUserId = Options.LocalUserId;
+
+			return EOS_Async<EOS_Achievements_OnQueryPlayerAchievementsCompleteCallbackInfo>(EOS_Achievements_QueryPlayerAchievements, AchievementsHandle, Options);
+		})
+		.Then([this](TOnlineAsyncOp<FQueryAchievementStates>& InAsyncOp, const EOS_Achievements_OnQueryPlayerAchievementsCompleteCallbackInfo* Data)
+		{
+			if (Data->ResultCode != EOS_EResult::EOS_Success)
+			{
+				UE_LOG(LogTemp, Warning, TEXT("EOS_Achievements_QueryPlayerAchievements failed with result=[%s]"), *LexToString(Data->ResultCode));
+
+				InAsyncOp.SetError(FromEOSError(Data->ResultCode));
+				return;
+			}
+
+			EOS_Achievements_GetPlayerAchievementCountOptions GetCountOptions = {};
+			GetCountOptions.ApiVersion = EOS_ACHIEVEMENTS_GETPLAYERACHIEVEMENTCOUNT_API_LATEST;
+			static_assert(EOS_ACHIEVEMENTS_GETPLAYERACHIEVEMENTCOUNT_API_LATEST == 1, "EOS_Achievements_GetPlayerAchievementCountOptions updated, check new fields");
+			GetCountOptions.UserId = GetProductUserIdChecked(InAsyncOp.GetParams().LocalUserId);
+
+			const uint32_t AchievementCount = EOS_Achievements_GetPlayerAchievementCount(AchievementsHandle, &GetCountOptions);
+
+			FAchievementStateMap NewAchievementStates;
+			for (uint32_t AchievementIdx = 0; AchievementIdx < AchievementCount; AchievementIdx++)
+			{
+				EOS_Achievements_CopyPlayerAchievementByIndexOptions CopyOptions = {};
+				CopyOptions.ApiVersion = EOS_ACHIEVEMENTS_COPYPLAYERACHIEVEMENTBYINDEX_API_LATEST;
+				static_assert(EOS_ACHIEVEMENTS_COPYPLAYERACHIEVEMENTBYINDEX_API_LATEST == 2, "EOS_Achievements_CopyPlayerAchievementByIndexOptions updated, check new fields");
+				CopyOptions.LocalUserId = GetProductUserIdChecked(InAsyncOp.GetParams().LocalUserId);
+				CopyOptions.TargetUserId = CopyOptions.LocalUserId;
+				CopyOptions.AchievementIndex = AchievementIdx;
+
+				EOS_Achievements_PlayerAchievement* EOSPlayerAchievement = nullptr;
+				const EOS_EResult CopyResult = EOS_Achievements_CopyPlayerAchievementByIndex(AchievementsHandle, &CopyOptions, &EOSPlayerAchievement);
+				if (CopyResult != EOS_EResult::EOS_Success)
+				{
+					UE_LOG(LogTemp, Warning, TEXT("EOS_Achievements_CopyPlayerAchievementByIndex failed with result=[%s]"), *LexToString(CopyResult));
+
+					InAsyncOp.SetError(FromEOSError(CopyResult));
+					return;
+				}
+				
+				static_assert(EOS_ACHIEVEMENTS_PLAYERACHIEVEMENT_API_LATEST == 2, "EOS_Achievements_PlayerAchievement updated, check new fields");
+				if (!ensure(EOSPlayerAchievement->ApiVersion == EOS_ACHIEVEMENTS_PLAYERACHIEVEMENT_API_LATEST))
+				{
+					UE_LOG(LogTemp, Warning, TEXT("EOS_Achievements_PlayerAchievement version mismatch Expected=%d Actual=%d"), EOS_ACHIEVEMENTS_PLAYERACHIEVEMENT_API_LATEST, EOSPlayerAchievement->ApiVersion);
+				}
+
+				FString AchievementId = UTF8_TO_TCHAR(EOSPlayerAchievement->AchievementId);
+				FAchievementState& AchievementState = NewAchievementStates.Emplace(AchievementId);
+				AchievementState.AchievementId = MoveTemp(AchievementId);
+				AchievementState.Progress = EOSPlayerAchievement->Progress;
+				AchievementState.UnlockTime = FDateTime::FromUnixTimestamp(EOSPlayerAchievement->UnlockTime);
+
+				EOS_Achievements_PlayerAchievement_Release(EOSPlayerAchievement);
+			}
+			AchievementStates.Emplace(InAsyncOp.GetParams().LocalUserId, MoveTemp(NewAchievementStates));
+			InAsyncOp.SetResult({});
+		})
+		.Enqueue(GetSerialQueue());
+	}
+	else
+	{
+		Op->SetResult({});
+	}
+
+	return Op->GetHandle();
+}
+
+TOnlineResult<FGetAchievementState> FAchievementsEOSGS::GetAchievementState(FGetAchievementState::Params&& Params)
+{
+	if (!Services.Get<FAuthEOSGS>()->IsLoggedIn(Params.LocalUserId))
+	{
+		return TOnlineResult<FGetAchievementState>(Errors::InvalidUser());
+	}
+
+	const FAchievementStateMap* LocalUserAchievementStates = AchievementStates.Find(Params.LocalUserId);
+	if (!LocalUserAchievementStates)
+	{
+		// Call QueryAchievementStates first
+		return TOnlineResult<FGetAchievementState>(Errors::InvalidState());
+	}
+
+	const FAchievementState* AchievementState = LocalUserAchievementStates->Find(Params.AchievementId);
+	if (!AchievementState)
+	{
+		return TOnlineResult<FGetAchievementState>(Errors::NotFound());
+	}
+
+	return TOnlineResult<FGetAchievementState>({ *AchievementState });
+}
+
+TOnlineAsyncOpHandle<FUnlockAchievements> FAchievementsEOSGS::UnlockAchievements(FUnlockAchievements::Params&& Params)
+{
+	TOnlineAsyncOpRef<FUnlockAchievements> Op = GetOp<FUnlockAchievements>(MoveTemp(Params));
+
+	if (!Services.Get<FAuthEOSGS>()->IsLoggedIn(Op->GetParams().LocalUserId))
+	{
+		Op->SetError(Errors::InvalidUser());
+		return Op->GetHandle();
+	}
+
+	if (Op->GetParams().AchievementIds.IsEmpty())
+	{
+		Op->SetError(Errors::InvalidParams());
+		return Op->GetHandle();
+	}
+
+	FAchievementStateMap* LocalUserAchievementStates = AchievementStates.Find(Op->GetParams().LocalUserId);
+	if (!LocalUserAchievementStates)
+	{
+		// Call QueryAchievementStates first
+		Op->SetError(Errors::InvalidState());
+		return Op->GetHandle();
+	}
+
+	for (const FString& AchievementId : Op->GetParams().AchievementIds)
+	{
+		const FAchievementState* AchievementState = LocalUserAchievementStates->Find(AchievementId);
+		if (!AchievementState)
+		{
+			Op->SetError(Errors::NotFound());
+			return Op->GetHandle();
+		}
+		if (AchievementState->Progress == 1.0f)
+		{
+			Op->SetError(Errors::Achievements::AlreadyUnlocked());
+			return Op->GetHandle();
+		}
+	}
+
+	Op->Then([this](TOnlineAsyncOp<FUnlockAchievements>& InAsyncOp)
+	{
+		const TArray<FString>& InAchievementIds = InAsyncOp.GetParams().AchievementIds;
+
+		TArray<FTCHARToUTF8> AchievementIdConverters;
+		AchievementIdConverters.Reserve(InAchievementIds.Num());
+		TArray<const char*> AchievementIdPtrs;
+		AchievementIdPtrs.Reserve(InAchievementIds.Num());
+		for (const FString& AchievementId : InAchievementIds)
+		{
+			FTCHARToUTF8& Converter = AchievementIdConverters.Emplace_GetRef(GetData(AchievementId), GetNum(AchievementId));
+			AchievementIdPtrs.Emplace(Converter.Get());
+		}
+
+		EOS_Achievements_UnlockAchievementsOptions Options = {};
+		Options.ApiVersion = EOS_ACHIEVEMENTS_UNLOCKACHIEVEMENTS_API_LATEST;
+		static_assert(EOS_ACHIEVEMENTS_UNLOCKACHIEVEMENTS_API_LATEST == 1, "EOS_Achievements_UnlockAchievementsOptions updated, check new fields");
+		Options.UserId = GetProductUserIdChecked(InAsyncOp.GetParams().LocalUserId);
+		Options.AchievementIds = AchievementIdPtrs.GetData();
+		Options.AchievementsCount = AchievementIdPtrs.Num();
+
+		return EOS_Async<EOS_Achievements_OnUnlockAchievementsCompleteCallbackInfo>(EOS_Achievements_UnlockAchievements, AchievementsHandle, Options);
+	})
+	.Then([this](TOnlineAsyncOp<FUnlockAchievements>& InAsyncOp, const EOS_Achievements_OnUnlockAchievementsCompleteCallbackInfo* Data)
+	{
+		if (Data->ResultCode != EOS_EResult::EOS_Success)
+		{
+			UE_LOG(LogTemp, Warning, TEXT("EOS_Achievements_UnlockAchievements failed with result=[%s]"), *LexToString(Data->ResultCode));
+
+			InAsyncOp.SetError(FromEOSError(Data->ResultCode));
+		}
+		else
+		{
+			InAsyncOp.SetResult({});
+		}
+
+	})
+	.Enqueue(GetSerialQueue());
+
+	return Op->GetHandle();
+}
+
+/* UE::Online */ }
diff --git a/Engine/Plugins/Online/Experimental/OnlineServicesEOSGS/Source/Private/Online/OnlineServicesEOSGS.cpp b/Engine/Plugins/Online/Experimental/OnlineServicesEOSGS/Source/Private/Online/OnlineServicesEOSGS.cpp
--- a/Engine/Plugins/Online/Experimental/OnlineServicesEOSGS/Source/Private/Online/OnlineServicesEOSGS.cpp
+++ b/Engine/Plugins/Online/Experimental/OnlineServicesEOSGS/Source/Private/Online/OnlineServicesEOSGS.cpp
@@ -2,6 +2,7 @@
 
 #include "Online/OnlineServicesEOSGS.h"
 
+#include "Online/AchievementsEOSGS.h"
 #include "Online/AuthEOSGS.h"
 #include "Online/LobbiesEOSGS.h"
 #include "Online/OnlineIdEOSGS.h"
@@ -47,6 +48,7 @@ FOnlineServicesEOSGS::FOnlineServicesEOSGS(FName InInstanceName)
 
 void FOnlineServicesEOSGS::RegisterComponents()
 {
+	Components.Register<FAchievementsEOSGS>(*this);
 	Components.Register<FAuthEOSGS>(*this);
 	Components.Register<FLobbiesEOSGS>(*this);
 	Super::RegisterComponents();
diff --git a/Engine/Plugins/Online/Experimental/OnlineServicesEOSGS/Source/Public/Online/AchievementsEOSGS.h b/Engine/Plugins/Online/Experimental/OnlineServicesEOSGS/Source/Public/Online/AchievementsEOSGS.h
new file mode 100644
--- /dev/null
+++ b/Engine/Plugins/Online/Experimental/OnlineServicesEOSGS/Source/Public/Online/AchievementsEOSGS.h
@@ -0,0 +1,49 @@
+// Copyright Epic Games, Inc. All Rights Reserved.
+
+#pragma once
+
+#include "CoreMinimal.h"
+#include "Online/AchievementsCommon.h"
+
+#if defined(EOS_PLATFORM_BASE_FILE_NAME)
+#include EOS_PLATFORM_BASE_FILE_NAME
+#endif
+#include "eos_achievements_types.h"
+
+namespace UE::Online {
+
+class FOnlineServicesEOSGS;
+
+class ONLINESERVICESEOSGS_API FAchievementsEOSGS : public FAchievementsCommon
+{
+public:
+	using Super = FAchievementsCommon;
+
+	FAchievementsEOSGS(FOnlineServicesEOSGS& InOwningSubsystem);
+	virtual ~FAchievementsEOSGS() = default;
+
+	// IAuth
+	virtual void Initialize() override;
+	virtual void Shutdown() override;
+
+	// IAchievements
+	virtual TOnlineAsyncOpHandle<FQueryAchievementDefinitions> QueryAchievementDefinitions(FQueryAchievementDefinitions::Params&& Params) override;
+	virtual TOnlineResult<FGetAchievementIds> GetAchievementIds(FGetAchievementIds::Params&& Params) override;
+	virtual TOnlineResult<FGetAchievementDefinition> GetAchievementDefinition(FGetAchievementDefinition::Params&& Params) override;
+	virtual TOnlineAsyncOpHandle<FQueryAchievementStates> QueryAchievementStates(FQueryAchievementStates::Params&& Params) override;
+	virtual TOnlineResult<FGetAchievementState> GetAchievementState(FGetAchievementState::Params&& Params) override;
+	virtual TOnlineAsyncOpHandle<FUnlockAchievements> UnlockAchievements(FUnlockAchievements::Params&& Params) override;
+
+protected:
+	EOS_HAchievements AchievementsHandle = nullptr;
+
+	using FAchievementDefinitionMap = TMap<FString, FAchievementDefinition>;
+	using FAchievementStateMap = TMap<FString, FAchievementState>;
+
+	TOptional<FAchievementDefinitionMap> AchievementDefinitions;
+	TMap<FOnlineAccountIdHandle, FAchievementStateMap> AchievementStates;
+
+	EOS_NotificationId NotifyAchievementsUnlockedNotificationId = EOS_INVALID_NOTIFICATIONID;
+};
+
+/* UE::Online */ }
