- Fix typo, those methods are part of the IOnlineComponent interface

diff --git a/Engine/Plugins/Online/Experimental/OnlineServicesEOSGS/Source/Public/Online/AchievementsEOSGS.h b/Engine/Plugins/Online/Experimental/OnlineServicesEOSGS/Source/Public/Online/AchievementsEOSGS.h
--- a/Engine/Plugins/Online/Experimental/OnlineServicesEOSGS/Source/Public/Online/AchievementsEOSGS.h
+++ b/Engine/Plugins/Online/Experimental/OnlineServicesEOSGS/Source/Public/Online/AchievementsEOSGS.h
@@ -22,7 +22,7 @@ public:
 	FAchievementsEOSGS(FOnlineServicesEOSGS& InOwningSubsystem);
 	virtual ~FAchievementsEOSGS() = default;
 
-	// IAuth
+	// IOnlineComponent
 	virtual void Initialize() override;
 	virtual void Shutdown() override;
 
