OSSv2: Add Stats EOS implementation.

[REDACTED_HANDLE] [REDACTED_HANDLE] [REDACTED_HANDLE] [REDACTED_HANDLE]
[REDACTED_HANDLE] OI-3873
[REDACTED_HANDLE] [REDACTED_HANDLE]

diff --git a/Engine/Plugins/Online/Experimental/OnlineServices/Source/OnlineServicesInterface/Public/Online/Stats.h b/Engine/Plugins/Online/Experimental/OnlineServices/Source/OnlineServicesInterface/Public/Online/Stats.h
--- a/Engine/Plugins/Online/Experimental/OnlineServices/Source/OnlineServicesInterface/Public/Online/Stats.h
+++ b/Engine/Plugins/Online/Experimental/OnlineServices/Source/OnlineServicesInterface/Public/Online/Stats.h
@@ -46,6 +46,8 @@ struct FQueryStats
 		FOnlineAccountIdHandle LocalUserId;
 		/* Target user id to query */
 		FOnlineAccountIdHandle TargetUserId;
+		/* The stats to query */
+		TArray<FString> StatNames;
 	};
 
 	struct Result
@@ -65,6 +67,8 @@ struct FBatchQueryStats
 		FOnlineAccountIdHandle LocalUserId;
 		/* The user ids to query */
 		TArray<FOnlineAccountIdHandle> TargetUserIds;
+		/* The stats to query */
+		TArray<FString> StatNames;
 	};
 
 	struct Result
diff --git a/Engine/Plugins/Online/Experimental/OnlineServicesEOSGS/Source/Private/Online/OnlineServicesEOSGS.cpp b/Engine/Plugins/Online/Experimental/OnlineServicesEOSGS/Source/Private/Online/OnlineServicesEOSGS.cpp
--- a/Engine/Plugins/Online/Experimental/OnlineServicesEOSGS/Source/Private/Online/OnlineServicesEOSGS.cpp
+++ b/Engine/Plugins/Online/Experimental/OnlineServicesEOSGS/Source/Private/Online/OnlineServicesEOSGS.cpp
@@ -5,6 +5,7 @@
 #include "Online/AchievementsEOSGS.h"
 #include "Online/AuthEOSGS.h"
 #include "Online/LobbiesEOSGS.h"
+#include "Online/StatsEOSGS.h"
 #include "Online/OnlineIdEOSGS.h"
 #include "Online/OnlineServicesEOSGSTypes.h"
 
@@ -51,6 +52,7 @@ void FOnlineServicesEOSGS::RegisterComponents()
 	Components.Register<FAchievementsEOSGS>(*this);
 	Components.Register<FAuthEOSGS>(*this);
 	Components.Register<FLobbiesEOSGS>(*this);
+	Components.Register<FStatsEOSGS>(*this);
 	Super::RegisterComponents();
 }
 
diff --git a/Engine/Plugins/Online/Experimental/OnlineServicesEOSGS/Source/Private/Online/StatsEOSGS.cpp b/Engine/Plugins/Online/Experimental/OnlineServicesEOSGS/Source/Private/Online/StatsEOSGS.cpp
new file mode 100644
--- /dev/null
+++ b/Engine/Plugins/Online/Experimental/OnlineServicesEOSGS/Source/Private/Online/StatsEOSGS.cpp
@@ -0,0 +1,255 @@
+// Copyright Epic Games, Inc. All Rights Reserved.
+
+#include "Online/StatsEOSGS.h"
+#include "Online/AuthEOSGS.h"
+#include "Online/OnlineIdEOSGS.h"
+#include "Online/OnlineServicesEOSGS.h"
+#include "Online/OnlineServicesEOSGSTypes.h"
+
+#include "eos_stats.h"
+
+namespace UE::Online {
+
+#define STAT_NAME_MAX_LENGTH_EOS 256 + 1 // 256 plus null terminator
+
+struct FStatNameRawBufferEOS
+{
+	char StatName[STAT_NAME_MAX_LENGTH_EOS];
+};
+
+FStatsEOSGS::FStatsEOSGS(FOnlineServicesEOSGS& InServices)
+	: Super(InServices)
+{
+}
+
+void FStatsEOSGS::Initialize()
+{
+	Super::Initialize();
+
+	StatsHandle = EOS_Platform_GetStatsInterface(static_cast<FOnlineServicesEOSGS&>(GetServices()).GetEOSPlatformHandle());
+	check(StatsHandle);
+}
+
+namespace Private
+{
+
+int32 ConvertToIntegerFromStatValue(const FStatValue& StatValue)
+{
+	int32 Value = 0;
+	// TODO: Convert to proper value when SchemaVariant is ready
+	return Value;
+}
+
+}
+
+TOnlineAsyncOpHandle<FUpdateStats> FStatsEOSGS::UpdateStats(FUpdateStats::Params&& Params)
+{
+	TOnlineAsyncOpRef<FUpdateStats> Op = GetOp<FUpdateStats>(MoveTemp(Params));
+
+	if (!Services.Get<FAuthEOSGS>()->IsLoggedIn(Op->GetParams().LocalUserId))
+	{
+		Op->SetError(Errors::InvalidUser());
+		return Op->GetHandle();
+	}
+
+	// TODO: Only load stats definitions on platforms which requires to upload leaderboard score separately
+
+	for (const FUserStats& UpdateUserStats : Op->GetParams().UpdateUsersStats)
+	{
+		Op->Then([this, &UpdateUserStats](TOnlineAsyncOp<FUpdateStats>& InAsyncOp)
+		{
+			TArray<EOS_Stats_IngestData> EOSData;
+			TArray<FStatNameRawBufferEOS> EOSStatNames;
+
+			// Preallocate all of the memory
+			EOSData.AddZeroed(UpdateUserStats.Stats.Num());
+			EOSStatNames.AddZeroed(UpdateUserStats.Stats.Num());
+
+			uint32 Index = 0;
+			for (const TPair<FString, FStatValue>& StatPair : UpdateUserStats.Stats)
+			{
+				EOS_Stats_IngestData& EOSStat = EOSData[Index];
+				EOSStat.ApiVersion = EOS_STATS_INGESTDATA_API_LATEST;
+				static_assert(EOS_STATS_INGESTDATA_API_LATEST == 1, "EOS_Stats_IngestData updated, check new fields");
+
+				EOSStat.IngestAmount = Private::ConvertToIntegerFromStatValue(StatPair.Value);
+				FCStringAnsi::Strncpy(EOSStatNames[Index].StatName, TCHAR_TO_UTF8(*StatPair.Key), STAT_NAME_MAX_LENGTH_EOS);
+				EOSStat.StatName = EOSStatNames[Index].StatName;
+
+				Index++;
+			}
+
+			EOS_Stats_IngestStatOptions Options = { };
+			Options.ApiVersion = EOS_STATS_INGESTSTAT_API_LATEST;
+			static_assert(EOS_STATS_INGESTSTAT_API_LATEST == 3, "EOS_Stats_IngestStat updated, check new fields");
+			Options.LocalUserId = GetProductUserIdChecked(InAsyncOp.GetParams().LocalUserId);
+			Options.TargetUserId = GetProductUserIdChecked(UpdateUserStats.UserId);
+			Options.Stats = EOSData.GetData();
+			Options.StatsCount = EOSData.Num();
+
+			return EOS_Async<EOS_Stats_IngestStatCompleteCallbackInfo>(EOS_Stats_IngestStat, StatsHandle, Options);
+		})
+		.Then([this](TOnlineAsyncOp<FUpdateStats>& InAsyncOp, const EOS_Stats_IngestStatCompleteCallbackInfo* Data)
+		{
+			if (Data->ResultCode != EOS_EResult::EOS_Success)
+			{
+				UE_LOG(LogTemp, Warning, TEXT("EOS_Stats_IngestStat failed with result=[%s]"), *LexToString(Data->ResultCode));
+				InAsyncOp.SetError(FromEOSError(Data->ResultCode));
+				return;
+			}
+		});
+	}
+
+	Op->Then([this](TOnlineAsyncOp<FUpdateStats>& InAsyncOp)
+	{
+		InAsyncOp.SetResult({});
+	});
+
+	Op->Enqueue(GetSerialQueue());
+
+	return Op->GetHandle();
+}
+
+namespace Private
+{
+
+TFuture<const EOS_Stats_OnQueryStatsCompleteCallbackInfo*> QueryStatsEOS(EOS_HStats StatsHandle, const FOnlineAccountIdHandle& LocalUserId, const FOnlineAccountIdHandle& TargetUserId, const TArray<FString>& StatNames)
+{
+	EOS_Stats_QueryStatsOptions Options;
+	Options.ApiVersion = EOS_STATS_QUERYSTATS_API_LATEST;
+	static_assert(EOS_STATS_INGESTSTAT_API_LATEST == 3, "EOS_Stats_QueryStats updated, check new fields");
+	Options.StartTime = EOS_STATS_TIME_UNDEFINED;
+	Options.EndTime = EOS_STATS_TIME_UNDEFINED;
+
+	TArray<FStatNameRawBufferEOS> EOSStatNames;
+	EOSStatNames.AddZeroed(StatNames.Num());
+
+	uint32 Index = 0;
+	for (const FString& StatName : StatNames)
+	{
+		FCStringAnsi::Strncpy(EOSStatNames[Index].StatName, TCHAR_TO_UTF8(*StatName), STAT_NAME_MAX_LENGTH_EOS);
+		Index++;
+	}
+
+	Options.StatNames = (const char**)EOSStatNames.GetData();
+	Options.StatNamesCount = StatNames.Num();
+	Options.LocalUserId = GetProductUserIdChecked(LocalUserId);
+	Options.TargetUserId = GetProductUserIdChecked(TargetUserId);
+
+	return EOS_Async<EOS_Stats_OnQueryStatsCompleteCallbackInfo>(EOS_Stats_QueryStats, StatsHandle, Options);
+}
+
+void ReadStatsFromEOSResult(EOS_HStats StatsHandle, const EOS_Stats_OnQueryStatsCompleteCallbackInfo* Data, const TArray<FString>& StatNames, TMap<FString, FStatValue>& OutStats)
+{
+	for (const FString& StatName : StatNames)
+	{
+		char StatNameANSI[STAT_NAME_MAX_LENGTH_EOS];
+		EOS_Stats_CopyStatByNameOptions Options = { };
+		Options.ApiVersion = EOS_STATS_COPYSTATBYNAME_API_LATEST;
+		static_assert(EOS_STATS_COPYSTATBYNAME_API_LATEST == 1, "EOS_Stats_CopyStatByNameOptions updated, check new fields");
+		Options.TargetUserId = Data->TargetUserId;
+		Options.Name = StatNameANSI;
+		FCStringAnsi::Strncpy(StatNameANSI, TCHAR_TO_UTF8(*StatName), STAT_NAME_MAX_LENGTH_EOS);
+
+		EOS_Stats_Stat* ReadStat = nullptr;
+		if (EOS_Stats_CopyStatByName(StatsHandle, &Options, &ReadStat) == EOS_EResult::EOS_Success)
+		{
+			// TODO: Add the stats here when SchemaVariant is ready
+			//OutStats.Add(StatName, FStatValue(ReadStat->Value));
+			EOS_Stats_Stat_Release(ReadStat);
+		}
+		else
+		{
+			// Put an empty stat in
+			UE_LOG(LogTemp, VeryVerbose, TEXT("Value not found for stat %s, adding empty value"), *StatName);
+			OutStats.Add(StatName, FStatValue());
+		}
+	}
+}
+
+}
+
+TOnlineAsyncOpHandle<FQueryStats> FStatsEOSGS::QueryStats(FQueryStats::Params&& Params)
+{
+	TOnlineAsyncOpRef<FQueryStats> Op = GetOp<FQueryStats>(MoveTemp(Params));
+
+	if (!Services.Get<FAuthEOSGS>()->IsLoggedIn(Op->GetParams().LocalUserId))
+	{
+		Op->SetError(Errors::InvalidUser());
+		return Op->GetHandle();
+	}
+
+	Op->Then([this](TOnlineAsyncOp<FQueryStats>& InAsyncOp)
+	{
+		return Private::QueryStatsEOS(StatsHandle, InAsyncOp.GetParams().LocalUserId, InAsyncOp.GetParams().TargetUserId, InAsyncOp.GetParams().StatNames);
+	})
+	.Then([this](TOnlineAsyncOp<FQueryStats>& InAsyncOp, const EOS_Stats_OnQueryStatsCompleteCallbackInfo* Data)
+	{
+		if (Data->ResultCode != EOS_EResult::EOS_Success)
+		{
+			UE_LOG(LogTemp, Warning, TEXT("EOS_Stats_QueryStats failed with result=[%s]"), *LexToString(Data->ResultCode));
+			InAsyncOp.SetError(FromEOSError(Data->ResultCode));
+			return;
+		}
+
+		FQueryStats::Result Result;
+		Private::ReadStatsFromEOSResult(StatsHandle, Data, InAsyncOp.GetParams().StatNames, Result.Stats);
+		InAsyncOp.SetResult(MoveTemp(Result));
+	});
+
+	Op->Enqueue(GetSerialQueue());
+
+	return Op->GetHandle();
+}
+
+TOnlineAsyncOpHandle<FBatchQueryStats> FStatsEOSGS::BatchQueryStats(FBatchQueryStats::Params&& Params)
+{
+	TOnlineAsyncOpRef<FBatchQueryStats> Op = GetOp<FBatchQueryStats>(MoveTemp(Params));
+
+	if (!Services.Get<FAuthEOSGS>()->IsLoggedIn(Op->GetParams().LocalUserId))
+	{
+		Op->SetError(Errors::InvalidUser());
+		return Op->GetHandle();
+	}
+
+	if (Op->GetParams().TargetUserIds.IsEmpty() || Op->GetParams().StatNames.IsEmpty())
+	{
+		Op->SetError(Errors::InvalidParams());
+		return Op->GetHandle();
+	}
+
+	for (const FOnlineAccountIdHandle& TargetUserId : Op->GetParams().TargetUserIds)
+	{
+		Op->Then([this, &TargetUserId](TOnlineAsyncOp<FBatchQueryStats>& InAsyncOp)
+		{
+			return Private::QueryStatsEOS(StatsHandle, InAsyncOp.GetParams().LocalUserId, TargetUserId, InAsyncOp.GetParams().StatNames);
+		})
+		.Then([this](TOnlineAsyncOp<FBatchQueryStats>& InAsyncOp, const EOS_Stats_OnQueryStatsCompleteCallbackInfo* Data)
+		{
+			if (Data->ResultCode != EOS_EResult::EOS_Success)
+			{
+				UE_LOG(LogTemp, Warning, TEXT("EOS_Stats_QueryStats failed with result=[%s]"), *LexToString(Data->ResultCode));
+				CachedUsersStats.Reset();
+				InAsyncOp.SetError(FromEOSError(Data->ResultCode));
+				return;
+			}
+
+			FUserStats& UserStats = CachedUsersStats.Emplace_GetRef();
+
+			Private::ReadStatsFromEOSResult(StatsHandle, Data, InAsyncOp.GetParams().StatNames, UserStats.Stats);
+		});
+	}
+
+	Op->Then([this](TOnlineAsyncOp<FBatchQueryStats>& InAsyncOp)
+	{
+		FBatchQueryStats::Result Result;
+		Result.UsersStats = MoveTemp(CachedUsersStats);
+		InAsyncOp.SetResult(MoveTemp(Result));
+	});
+
+	Op->Enqueue(GetSerialQueue());
+
+	return Op->GetHandle();
+}
+
+/* UE::Online */ }
diff --git a/Engine/Plugins/Online/Experimental/OnlineServicesEOSGS/Source/Public/Online/StatsEOSGS.h b/Engine/Plugins/Online/Experimental/OnlineServicesEOSGS/Source/Public/Online/StatsEOSGS.h
new file mode 100644
--- /dev/null
+++ b/Engine/Plugins/Online/Experimental/OnlineServicesEOSGS/Source/Public/Online/StatsEOSGS.h
@@ -0,0 +1,37 @@
+// Copyright Epic Games, Inc. All Rights Reserved.
+
+#pragma once
+
+#include "CoreMinimal.h"
+#include "Online/StatsCommon.h"
+
+#if defined(EOS_PLATFORM_BASE_FILE_NAME)
+#include EOS_PLATFORM_BASE_FILE_NAME
+#endif
+#include "eos_stats_types.h"
+
+namespace UE::Online {
+
+class FOnlineServicesEOSGS;
+
+class ONLINESERVICESEOSGS_API FStatsEOSGS : public FStatsCommon
+{
+public:
+	using Super = FStatsCommon;
+
+	FStatsEOSGS(FOnlineServicesEOSGS& InOwningSubsystem);
+
+	// TOnlineComponent
+	virtual void Initialize() override;
+
+	// IStats
+	virtual TOnlineAsyncOpHandle<FUpdateStats> UpdateStats(FUpdateStats::Params&& Params) override;
+	virtual TOnlineAsyncOpHandle<FQueryStats> QueryStats(FQueryStats::Params&& Params) override;
+	virtual TOnlineAsyncOpHandle<FBatchQueryStats> BatchQueryStats(FBatchQueryStats::Params&& Params) override;
+
+private:
+	EOS_HStats StatsHandle = nullptr;
+	TArray<FUserStats> CachedUsersStats;
+};
+
+/* UE::Online */ }
