OSSv2 Sessions Interface - EOSGS Implementation - Basic setup and Create, Update, Leave functions

diff --git a/Engine/Plugins/Online/Experimental/OnlineServicesEOSGS/Source/OnlineServicesEOSGS.Build.cs b/Engine/Plugins/Online/Experimental/OnlineServicesEOSGS/Source/OnlineServicesEOSGS.Build.cs
--- a/Engine/Plugins/Online/Experimental/OnlineServicesEOSGS/Source/OnlineServicesEOSGS.Build.cs
+++ b/Engine/Plugins/Online/Experimental/OnlineServicesEOSGS/Source/OnlineServicesEOSGS.Build.cs
@@ -20,7 +20,8 @@ public class OnlineServicesEOSGS : ModuleRules
 				"CoreOnline",
 				"CoreUObject",
 				"EOSSDK",
-				"EOSShared"
+				"EOSShared",
+				"Sockets"
 			}
 		);
 
@@ -29,8 +30,7 @@ public class OnlineServicesEOSGS : ModuleRules
 			PrivateDependencyModuleNames.AddRange(
 				new string[]
 				{
-					"SocketSubsystemEOS",
-					"Sockets"
+					"SocketSubsystemEOS"
 				}
 			);
 		}
diff --git a/Engine/Plugins/Online/Experimental/OnlineServicesEOSGS/Source/Private/Online/OnlineServicesEOSGS.cpp b/Engine/Plugins/Online/Experimental/OnlineServicesEOSGS/Source/Private/Online/OnlineServicesEOSGS.cpp
--- a/Engine/Plugins/Online/Experimental/OnlineServicesEOSGS/Source/Private/Online/OnlineServicesEOSGS.cpp
+++ b/Engine/Plugins/Online/Experimental/OnlineServicesEOSGS/Source/Private/Online/OnlineServicesEOSGS.cpp
@@ -6,6 +6,7 @@
 #include "Online/AuthEOSGS.h"
 #include "Online/LobbiesEOSGS.h"
 #include "Online/StatsEOSGS.h"
+#include "Online/SessionsEOSGS.h"
 #include "Online/OnlineIdEOSGS.h"
 #include "Online/OnlineServicesEOSGSTypes.h"
 
@@ -53,6 +54,7 @@ void FOnlineServicesEOSGS::RegisterComponents()
 	Components.Register<FAuthEOSGS>(*this);
 	Components.Register<FLobbiesEOSGS>(*this);
 	Components.Register<FStatsEOSGS>(*this);
+	Components.Register<FSessionsEOSGS>(*this);
 	Super::RegisterComponents();
 }
 
@@ -120,33 +122,67 @@ void FOnlineServicesEOSGS::Initialize()
 
 TOnlineResult<FGetResolvedConnectString> FOnlineServicesEOSGS::GetResolvedConnectString(FGetResolvedConnectString::Params&& Params)
 {
-	ILobbiesPtr LobbiesEOS = GetLobbiesInterface();
-	check(LobbiesEOS);
-
-	TOnlineResult<FGetJoinedLobbies> JoinedLobbies = LobbiesEOS->GetJoinedLobbies({ Params.LocalUserId });
-	if (JoinedLobbies.IsOk())
+	if (Params.LobbyId.IsValid())
 	{
-		for (TSharedRef<const FLobby>& Lobby : JoinedLobbies.GetOkValue().Lobbies)
+		ILobbiesPtr LobbiesEOS = GetLobbiesInterface();
+		check(LobbiesEOS);
+
+		TOnlineResult<FGetJoinedLobbies> JoinedLobbies = LobbiesEOS->GetJoinedLobbies({ Params.LocalUserId });
+		if (JoinedLobbies.IsOk())
 		{
-			if (Lobby->LobbyId == Params.LobbyId)
+			for (TSharedRef<const FLobby>& Lobby : JoinedLobbies.GetOkValue().Lobbies)
 			{
+				if (Lobby->LobbyId == Params.LobbyId)
+				{
 #if WITH_ENGINE
- 				//It should look like this: "EOS:0002aeeb5b2d4388a3752dd6d31222ec:GameNetDriver:97"
- 				FString NetDriverName = GetDefault<UNetDriverEOSBase>()->NetDriverName.ToString();
- 				FInternetAddrEOS TempAddr(GetProductUserIdChecked(Lobby->OwnerAccountId), NetDriverName, GetTypeHash(NetDriverName));
- 				return TOnlineResult<FGetResolvedConnectString>({ TempAddr.ToString(true) });
+					//It should look like this: "EOS:0002aeeb5b2d4388a3752dd6d31222ec:GameNetDriver:97"
+					FString NetDriverName = GetDefault<UNetDriverEOSBase>()->NetDriverName.ToString();
+					FInternetAddrEOS TempAddr(GetProductUserIdChecked(Lobby->OwnerAccountId), NetDriverName, GetTypeHash(NetDriverName));
+					return TOnlineResult<FGetResolvedConnectString>({ TempAddr.ToString(true) });
 #else
-				return TOnlineResult<FGetResolvedConnectString>(Errors::NotImplemented());
+					return TOnlineResult<FGetResolvedConnectString>(Errors::NotImplemented());
 #endif
+				}
 			}
+
+			// No matching lobby
+			return TOnlineResult<FGetResolvedConnectString>(Errors::NotFound());
+		}
+		else
+		{
+			return TOnlineResult<FGetResolvedConnectString>(JoinedLobbies.GetErrorValue());
 		}
-		// No matching lobby
-		return TOnlineResult<FGetResolvedConnectString>(Errors::InvalidParams());
 	}
-	else
+	else if (Params.SessionId.IsValid())
 	{
-		return TOnlineResult<FGetResolvedConnectString>(JoinedLobbies.GetErrorValue());
+		ISessionsPtr SessionsEOS = GetSessionsInterface();
+		check(SessionsEOS);
+
+		TOnlineResult<FGetSessionById> Result = SessionsEOS->GetSessionById({ Params.SessionId });
+		if (Result.IsOk())
+		{
+#if WITH_ENGINE
+			TSharedRef<const FSession> Session = Result.GetOkValue().Session;
+
+			//It should look like this: "EOS:0002aeeb5b2d4388a3752dd6d31222ec:GameNetDriver:97"
+			FString NetDriverName = GetDefault<UNetDriverEOSBase>()->NetDriverName.ToString();
+			FInternetAddrEOS TempAddr(GetProductUserIdChecked(Session->OwnerUserId), NetDriverName, GetTypeHash(NetDriverName));
+			return TOnlineResult<FGetResolvedConnectString>({ TempAddr.ToString(true) });
+#else
+			return TOnlineResult<FGetResolvedConnectString>(Errors::NotImplemented());
+#endif
+		}
+		else
+		{
+			return TOnlineResult<FGetResolvedConnectString>(Result.GetErrorValue());
+		}
+
+		// No matching session
+		return TOnlineResult<FGetResolvedConnectString>(Errors::NotFound());
 	}
+
+	// No valid lobby or session id set
+	return TOnlineResult<FGetResolvedConnectString>(Errors::InvalidParams());
 }
 
 EOS_HPlatform FOnlineServicesEOSGS::GetEOSPlatformHandle() const
diff --git a/Engine/Plugins/Online/Experimental/OnlineServicesEOSGS/Source/Private/Online/SessionsEOSGS.cpp b/Engine/Plugins/Online/Experimental/OnlineServicesEOSGS/Source/Private/Online/SessionsEOSGS.cpp
new file mode 100644
--- /dev/null
+++ b/Engine/Plugins/Online/Experimental/OnlineServicesEOSGS/Source/Private/Online/SessionsEOSGS.cpp
@@ -0,0 +1,786 @@
+// Copyright Epic Games, Inc. All Rights Reserved.
+
+#include "Online/SessionsEOSGS.h"
+
+#include "Online/AuthEOSGS.h"
+#include "Online/OnlineIdEOSGS.h"
+#include "Online/OnlineServicesEOSGS.h"
+#include "Online/OnlineServicesEOSGSTypes.h"
+
+#include "eos_sessions.h"
+#include "eos_sessions_types.h"
+
+namespace UE::Online {
+
+/** Auxiliary methods */
+
+EOS_EOnlineSessionPermissionLevel GetEOSGSType(const ESessionJoinPolicy& Value)
+{
+	switch (Value)
+	{
+	case ESessionJoinPolicy::Public:		return EOS_EOnlineSessionPermissionLevel::EOS_OSPF_PublicAdvertised;
+	case ESessionJoinPolicy::FriendsOnly:	return EOS_EOnlineSessionPermissionLevel::EOS_OSPF_JoinViaPresence;
+	case ESessionJoinPolicy::InviteOnly:	return EOS_EOnlineSessionPermissionLevel::EOS_OSPF_InviteOnly;
+	}
+
+	checkNoEntry();
+	return EOS_EOnlineSessionPermissionLevel::EOS_OSPF_InviteOnly;
+}
+
+/** FOnlineSessionIdRegistryEOSGS */
+
+FOnlineSessionIdRegistryEOSGS& FOnlineSessionIdRegistryEOSGS::Get()
+{
+	static FOnlineSessionIdRegistryEOSGS Instance;
+	return Instance;
+}
+
+FOnlineSessionIdHandle FOnlineSessionIdRegistryEOSGS::Find(const FString& SessionId) const
+{
+	if (const FOnlineSessionIdString* Entry = StringToId.Find(SessionId))
+	{
+		return Entry->Handle;
+	}
+	return FOnlineSessionIdHandle();
+}
+
+FOnlineSessionIdHandle FOnlineSessionIdRegistryEOSGS::FindOrAdd(const FString& SessionId)
+{
+	const FOnlineSessionIdString* Entry = StringToId.Find(SessionId);
+	if (Entry)
+	{
+		return Entry->Handle;
+	}
+
+	FOnlineSessionIdString& NewId = Ids.Emplace_GetRef();
+	NewId.Data = SessionId;
+	NewId.Handle = FOnlineSessionIdHandle(EOnlineServices::Epic, Ids.Num());
+
+	StringToId.Add(SessionId, NewId);
+
+	return NewId.Handle;
+}
+
+const FOnlineSessionIdString* FOnlineSessionIdRegistryEOSGS::GetInternal(const FOnlineSessionIdHandle& Handle) const
+{
+	if (Handle.IsValid() && Handle.GetOnlineServicesType() == EOnlineServices::Epic && Handle.GetHandle() <= (uint32)Ids.Num())
+	{
+		return &Ids[Handle.GetHandle() - 1];
+	}
+	return nullptr;
+}
+
+FString FOnlineSessionIdRegistryEOSGS::ToLogString(const FOnlineSessionIdHandle& Handle) const
+{
+	if (const FOnlineSessionIdString* Id = GetInternal(Handle))
+	{
+		return Id->Data;
+	}
+
+	return FString(TEXT("[InvalidSessionID]"));
+}
+
+TArray<uint8> FOnlineSessionIdRegistryEOSGS::ToReplicationData(const FOnlineSessionIdHandle& Handle) const
+{
+	if (const FOnlineSessionIdString* Id = GetInternal(Handle))
+	{
+		TArray<uint8> ReplicationData;
+		ReplicationData.SetNumUninitialized(Id->Data.Len());
+		StringToBytes(Id->Data, ReplicationData.GetData(), Id->Data.Len());
+		return ReplicationData;
+	}
+
+	return TArray<uint8>();;
+}
+
+FOnlineSessionIdHandle FOnlineSessionIdRegistryEOSGS::FromReplicationData(const TArray<uint8>& ReplicationData)
+{
+	FString Result = BytesToString(ReplicationData.GetData(), ReplicationData.Num());
+	if (Result.Len() > 0)
+	{
+		return FindOrAdd(Result);
+	}
+	return FOnlineSessionIdHandle();
+}
+
+/** FSessionEOSGS */
+
+FSessionEOSGS::FSessionEOSGS()
+	: FSessionEOSGS(FSession())
+{
+}
+
+FSessionEOSGS::FSessionEOSGS(const FSession& InSession)
+	: FSession(InSession)
+{
+}
+
+/** FSessionsEOSGS */
+
+FSessionsEOSGS::FSessionsEOSGS(FOnlineServicesEOSGS& InServices)
+	: Super(InServices)
+	, Services(InServices)
+{
+}
+
+void FSessionsEOSGS::Initialize()
+{
+	Super::Initialize();
+
+	SessionsHandle = EOS_Platform_GetSessionsInterface(static_cast<FOnlineServicesEOSGS&>(GetServices()).GetEOSPlatformHandle());
+	check(SessionsHandle);
+}
+
+TOnlineResult<FGetAllSessions> FSessionsEOSGS::GetAllSessions(FGetAllSessions::Params&& Params)
+{
+	FGetAllSessions::Result Result;
+
+	for (TMap<FName, TSharedRef<FSessionEOSGS>>::TConstIterator It(SessionsByName); It; ++It)
+	{
+		Result.Sessions.Add(It.Value());
+	}
+
+	return TOnlineResult<FGetAllSessions>(MoveTemp(Result));
+}
+
+TOnlineResult<FGetSessionByName> FSessionsEOSGS::GetSessionByName(FGetSessionByName::Params&& Params)
+{
+	if (TSharedRef<FSessionEOSGS>* FoundSession = SessionsByName.Find(Params.LocalName))
+	{
+		return TOnlineResult<FGetSessionByName>({ *FoundSession });
+	}
+	else
+	{
+		return TOnlineResult<FGetSessionByName>(Errors::NotFound());
+	}
+}
+
+TOnlineResult<FGetSessionById> FSessionsEOSGS::GetSessionById(FGetSessionById::Params&& Params)
+{
+	if (TSharedRef<FSessionEOSGS>* FoundSession = SessionsById.Find(Params.IdHandle))
+	{
+		return TOnlineResult<FGetSessionById>({ *FoundSession });
+	}
+	else
+	{
+		return TOnlineResult<FGetSessionById>(Errors::NotFound());
+	}
+}
+
+TOnlineAsyncOpHandle<FCreateSession> FSessionsEOSGS::CreateSession(FCreateSession::Params&& Params)
+{
+	TOnlineAsyncOpRef<FCreateSession> Op = GetOp<FCreateSession>(MoveTemp(Params));
+
+	Op->Then([this](TOnlineAsyncOp<FCreateSession>& Op) mutable
+	{
+		const FCreateSession::Params& OpParams = Op.GetParams();
+
+		if (TSharedRef<FSessionEOSGS>* FoundSession = SessionsByName.Find(OpParams.SessionName))
+		{
+			UE_LOG(LogTemp, Warning, TEXT("[FSessionsEOSGS::CreateSession] Could not create session with name [%s]. A session with that name already exists"), *OpParams.SessionName.ToString());
+
+			Op.SetError(Errors::InvalidState()); // TODO: New error: Session with name %s already exists
+
+			return MakeFulfilledPromise<TDefaultErrorResult<FUpdateSessionImpl>>().GetFuture();
+		}
+
+		if (OpParams.SessionSettings.NumMaxPrivateConnections <= 0 && OpParams.SessionSettings.NumMaxPublicConnections <= 0)
+		{
+			UE_LOG(LogTemp, Warning, TEXT("[FSessionsEOSGS::CreateSession] Could not create session with no valid NumMaxPrivateConnections [%d] or NumMaxPublicConnections [%d]"), OpParams.SessionSettings.NumMaxPrivateConnections, OpParams.SessionSettings.NumMaxPublicConnections);
+
+			Op.SetError(Errors::InvalidParams());
+
+			return MakeFulfilledPromise<TDefaultErrorResult<FUpdateSessionImpl>>().GetFuture();
+		}
+
+		if (!OpParams.SessionSettings.IsDedicatedServerSession)
+		{
+			IAuthPtr Auth = Services.GetAuthInterface();
+
+			FAuthGetAccountByAccountId::Params GetAccountByAccountIParams { OpParams.LocalUserId };
+			TOnlineResult<FAuthGetAccountByAccountId> Result = Auth->GetAccountByAccountId(MoveTemp(GetAccountByAccountIParams));
+
+			if (Result.IsOk())
+			{
+				if (Result.GetOkValue().AccountInfo->LoginStatus == ELoginStatus::NotLoggedIn)
+				{
+					UE_LOG(LogTemp, Warning, TEXT("[FSessionsEOSGS::CreateSession] Could not create session with user [%s] not logged in"), *ToLogString(OpParams.LocalUserId));
+
+					Op.SetError(Errors::InvalidUser());
+
+					return MakeFulfilledPromise<TDefaultErrorResult<FUpdateSessionImpl>>().GetFuture();
+				}
+			}
+		}
+
+		if (!OpParams.SessionSettings.SessionIdOverride.IsEmpty())
+		{
+			int32 Length = OpParams.SessionSettings.SessionIdOverride.Len();
+
+			if (Length < EOS_SESSIONMODIFICATION_MIN_SESSIONIDOVERRIDE_LENGTH || Length > EOS_SESSIONMODIFICATION_MAX_SESSIONIDOVERRIDE_LENGTH)
+			{
+				UE_LOG(LogTemp, Warning, TEXT("[FSessionsEOSGS::CreateSession] Could not create session with SessionIdOverride [%s] of size [%d]. SessionIdOverride size must be between [%d] and [%d] characters long"), *OpParams.SessionSettings.SessionIdOverride, Length, EOS_SESSIONMODIFICATION_MIN_SESSIONIDOVERRIDE_LENGTH, EOS_SESSIONMODIFICATION_MAX_SESSIONIDOVERRIDE_LENGTH);
+
+				Op.SetError(Errors::InvalidParams());
+
+				return MakeFulfilledPromise<TDefaultErrorResult<FUpdateSessionImpl>>().GetFuture();
+			}
+		}
+
+		// TODO: Check if bucket id is set. Depends on if the method can be called successfully without it
+
+		// After all initial checks, we start the session creation operations
+
+		EOS_Sessions_CreateSessionModificationOptions CreateSessionModificationOptions = {};
+		CreateSessionModificationOptions.ApiVersion = EOS_SESSIONS_CREATESESSIONMODIFICATION_API_LATEST;
+		static_assert(EOS_SESSIONS_CREATESESSIONMODIFICATION_API_LATEST == 4, "EOS_Sessions_CreateSessionModificationOptions updated, check new fields");
+
+		CreateSessionModificationOptions.bPresenceEnabled = OpParams.SessionSettings.bPresenceEnabled;
+		CreateSessionModificationOptions.bSanctionsEnabled = !OpParams.SessionSettings.bAllowSanctionedPlayers;
+
+		const FCustomSessionSetting* BucketIdSetting = OpParams.SessionSettings.CustomSettings.Find(EOS_SESSIONS_BUCKET_ID);
+		FTCHARToUTF8 BucketIdUtf8(BucketIdSetting != nullptr ? BucketIdSetting->Data.Get<FString>() : FString());
+
+		if (BucketIdUtf8.Length())
+		{
+			CreateSessionModificationOptions.BucketId = BucketIdUtf8.Get();
+		}
+	
+		CreateSessionModificationOptions.LocalUserId = GetProductUserIdChecked(OpParams.LocalUserId);
+		CreateSessionModificationOptions.MaxPlayers = OpParams.SessionSettings.NumMaxPrivateConnections + OpParams.SessionSettings.NumMaxPublicConnections;
+
+		FTCHARToUTF8 SessionIdUtf8(OpParams.SessionSettings.SessionIdOverride);
+		if (SessionIdUtf8.Length())
+		{
+			CreateSessionModificationOptions.SessionId = SessionIdUtf8.Get();
+		}
+
+		FTCHARToUTF8 SessionNameUtf8(OpParams.SessionName.ToString());
+		CreateSessionModificationOptions.SessionName = SessionNameUtf8.Get();
+
+		EOS_HSessionModification SessionModificationHandle = nullptr;
+		EOS_EResult ResultCode = EOS_Sessions_CreateSessionModification(SessionsHandle, &CreateSessionModificationOptions, &SessionModificationHandle);
+		if (ResultCode != EOS_EResult::EOS_Success)
+		{
+			UE_LOG(LogTemp, Error, TEXT("[FSessionsEOSGS::CreateSession] EOS_Sessions_CreateSessionModification failed with result [%s]"), *LexToString(ResultCode));
+
+			Op.SetError(FromEOSError(ResultCode));
+
+			return MakeFulfilledPromise<TDefaultErrorResult<FUpdateSessionImpl>>().GetFuture();
+		}
+
+		// TODO: We could call EOS_SessionModification_SetHostAddress at this point, although it's not necessary
+
+		// We write all SessionSettings values into the SessionModificationHandle
+		WriteCreateSessionModificationHandle(SessionModificationHandle, OpParams.SessionSettings);
+
+		// If the session modification handle is created successfully, we'll create the local session object
+		FSessionEOSGS NewSession = FSessionEOSGS();
+		NewSession.CurrentState = ESessionState::Creating;
+		NewSession.OwnerUserId = OpParams.LocalUserId;
+		NewSession.SessionSettings = OpParams.SessionSettings;
+
+ 		TSharedRef<FSessionEOSGS> NewSessionEOSGSRef = MakeShared<FSessionEOSGS>(NewSession);
+
+ 		SessionsByName.Add(OpParams.SessionName, NewSessionEOSGSRef);
+ 		SessionsById.Add(NewSession.SessionId, NewSessionEOSGSRef);
+
+		FUpdateSessionImpl::Params UpdateSessionImplParams { SessionModificationHandle };
+		
+		return UpdateSessionImpl(MoveTemp(UpdateSessionImplParams));
+	})
+	.Then([this](TOnlineAsyncOp<FCreateSession>& Op, TDefaultErrorResult<FUpdateSessionImpl>&& Result) mutable
+	{
+		if (Result.IsOk())
+		{
+			if (TSharedRef<FSessionEOSGS>* Session = SessionsByName.Find(Op.GetParams().SessionName))
+			{
+				Session->Get().CurrentState = ESessionState::Valid;
+
+				// TODO: Make additional local users join the created session
+
+				Op.SetResult(FCreateSession::Result{ *Session });
+			}
+			else
+			{
+				Op.SetError(Errors::NotFound());
+			}
+		}
+		else
+		{
+			Op.SetError(MoveTemp(Result.GetErrorValue()));
+		}
+	})
+	.Enqueue(GetSerialQueue());
+
+	return Op->GetHandle();
+}
+
+void FSessionsEOSGS::SetPermissionLevel(EOS_HSessionModification& SessionModificationHandle, const ESessionJoinPolicy& NewJoinPolicy)
+{
+	EOS_SessionModification_SetPermissionLevelOptions Options = { };
+	Options.ApiVersion = EOS_SESSIONMODIFICATION_SETPERMISSIONLEVEL_API_LATEST;
+	static_assert(EOS_SESSIONMODIFICATION_SETPERMISSIONLEVEL_API_LATEST == 1, "EOS_SessionModification_SetPermissionLevelOptions updated, check new fields");
+
+	Options.PermissionLevel = GetEOSGSType(NewJoinPolicy);
+
+	EOS_EResult ResultCode = EOS_SessionModification_SetPermissionLevel(SessionModificationHandle, &Options);
+	if (ResultCode != EOS_EResult::EOS_Success)
+	{
+		UE_LOG(LogTemp, Warning, TEXT("EOS_SessionModification_SetPermissionLevel failed with result [%s]"), *LexToString(ResultCode));
+	}
+}
+
+void FSessionsEOSGS::SetBucketId(EOS_HSessionModification& SessionModificationHandle, const FString& NewBucketId)
+{
+	EOS_SessionModification_SetBucketIdOptions Options = { };
+	Options.ApiVersion = EOS_SESSIONMODIFICATION_SETBUCKETID_API_LATEST;
+	static_assert(EOS_SESSIONMODIFICATION_SETBUCKETID_API_LATEST == 1, "EOS_SessionModification_SetBucketIdOptions updated, check new fields");
+
+	FTCHARToUTF8 BucketIdUtf8(*NewBucketId);
+	Options.BucketId = BucketIdUtf8.Get();
+
+	EOS_EResult ResultCode = EOS_SessionModification_SetBucketId(SessionModificationHandle, &Options);
+	if (ResultCode != EOS_EResult::EOS_Success)
+	{
+		UE_LOG(LogTemp, Warning, TEXT("EOS_SessionModification_SetMaxPlayers failed with result [%s]"), *LexToString(ResultCode));
+	}
+}
+
+void FSessionsEOSGS::SetMaxPlayers(EOS_HSessionModification& SessionModificationHandle, const int32& NewMaxPlayers)
+{
+	EOS_SessionModification_SetMaxPlayersOptions Options = { };
+	Options.ApiVersion = EOS_SESSIONMODIFICATION_SETMAXPLAYERS_API_LATEST;
+	static_assert(EOS_SESSIONMODIFICATION_SETMAXPLAYERS_API_LATEST == 1, "EOS_SessionModification_SetMaxPlayersOptions updated, check new fields");
+
+	Options.MaxPlayers = NewMaxPlayers;
+
+	EOS_EResult ResultCode = EOS_SessionModification_SetMaxPlayers(SessionModificationHandle, &Options);
+	if (ResultCode != EOS_EResult::EOS_Success)
+	{
+		UE_LOG(LogTemp, Warning, TEXT("EOS_SessionModification_SetMaxPlayers failed with result [%s]"), *LexToString(ResultCode));
+	}
+}
+
+void FSessionsEOSGS::WriteCreateSessionModificationHandle(EOS_HSessionModification& SessionModificationHandle, const FSessionSettings& SessionSettings)
+{
+	// TODO: We have the option to call EOS_SessionModification_SetHostAddress in EOS, useful if the session owner changes
+
+	// TODO: We either call Start or End session, EOS_SessionModification_SetJoinInProgressAllowed and EOS_SessionModification_SetInvitesAllowed depending on this value changing, or a custom setting.
+	// Investigation pending to see if there is overlay changes when the session is not joinable (started + JIP false, or InvitesAllowed false)
+
+	// TODO: Add non-custom settings as attributes on the API call to transmit that information in session searches. Example:
+
+// 	FCustomSessionSetting NewSetting;
+// 	NewSetting.Data.Set<bool>(SessionSettings.bAllowNewMembers);
+// 	NewSetting.Visibility = ECustomSessionSettingVisibility::ViaOnlineService;
+// 
+// 	AddAttribute(SessionModificationHandle, EOS_SESSIONS_ALLOW_NEW_MEMBERS, NewSetting);
+
+	SetPermissionLevel(SessionModificationHandle, SessionSettings.JoinPolicy);
+
+	SetMaxPlayers(SessionModificationHandle, SessionSettings.NumMaxPrivateConnections + SessionSettings.NumMaxPublicConnections);
+
+	// Custom Settings
+
+// 	for (FCustomSessionSettingsMap::TConstIterator It(SessionSettings.CustomSettings); It; ++It)
+// 	{
+// 		// TODO: Add the attribute to the SessionModification. Pending SchemaVariant work.
+// 	}
+
+	// BucketId has its own set method on the API
+	const FCustomSessionSetting* NewBucketIdSetting = SessionSettings.CustomSettings.Find(EOS_SESSIONS_BUCKET_ID);
+	if (NewBucketIdSetting)
+	{
+		SetBucketId(SessionModificationHandle, NewBucketIdSetting->Data.Get<FString>());
+	}
+
+	// Session Members
+
+// 	for (FSessionMembersMap::TConstIterator MemberIt(SessionSettings.SessionMembers); MemberIt; ++MemberIt)
+// 	{
+// 		const FSessionMember& SessionMember = MemberIt.Value();
+// 		
+// 		for (FCustomSessionSettingsMap::TConstIterator SettingIt(SessionMember.MemberSettings); SettingIt; ++SettingIt)
+// 		{
+// 			const FName& Key = FName(LexToString(GetProductUserIdChecked(MemberIt.Key())) + TEXT(":") + SettingIt.Key().ToString());
+// 			const FCustomSessionSetting& Value = SettingIt.Value();
+// 
+// 			// TODO: Add the attribute to the SessionModification. Pending SchemaVariant work.
+// 		}		
+// 	}
+}
+
+void FSessionsEOSGS::WriteUpdateSessionModificationHandle(EOS_HSessionModification& SessionModificationHandle, FSessionSettings& SessionSettings, const FSessionSettingsUpdate& NewSettings)
+{
+	// TODO: We have the option to call EOS_SessionModification_SetHostAddress in EOS, useful if the session owner changes
+	
+	// TODO: We either call Start or End session, EOS_SessionModification_SetJoinInProgressAllowed and EOS_SessionModification_SetInvitesAllowed depending on this value changing, or a custom setting.
+	// Investigation pending to see if there is overlay changes when the session is not joinable (started + JIP false, or InvitesAllowed false)
+
+	// TODO: Add non-custom settings as attributes on the API call to transmit that information in session searches. Example:
+
+// 	FCustomSessionSetting NewSetting;
+// 	NewSetting.Data.Set<bool>(SessionSettings.bAllowNewMembers);
+// 	NewSetting.Visibility = ECustomSessionSettingVisibility::ViaOnlineService;
+// 
+// 	AddAttribute(SessionModificationHandle, EOS_SESSIONS_ALLOW_NEW_MEMBERS, NewSetting);
+
+	if (NewSettings.bAllowNewMembers.IsSet())
+	{
+		SessionSettings.bAllowNewMembers = NewSettings.bAllowNewMembers.GetValue();
+	}
+
+	if (NewSettings.bAllowSanctionedPlayers.IsSet())
+	{
+		SessionSettings.bAllowSanctionedPlayers = NewSettings.bAllowSanctionedPlayers.GetValue();
+	}
+
+	if (NewSettings.bAllowUnregisteredPlayers.IsSet())
+	{
+		SessionSettings.bAllowUnregisteredPlayers = NewSettings.bAllowUnregisteredPlayers.GetValue();
+	}
+
+	if (NewSettings.bAntiCheatProtected.IsSet())
+	{
+		SessionSettings.bAntiCheatProtected = NewSettings.bAntiCheatProtected.GetValue();
+	}
+
+	if (NewSettings.bPresenceEnabled.IsSet())
+	{
+		SessionSettings.bPresenceEnabled = NewSettings.bPresenceEnabled.GetValue();
+	}
+
+	if (NewSettings.IsDedicatedServerSession.IsSet())
+	{
+		SessionSettings.IsDedicatedServerSession = NewSettings.IsDedicatedServerSession.GetValue();
+	}
+
+	if (NewSettings.IsLANSession.IsSet())
+	{
+		SessionSettings.IsLANSession = NewSettings.IsLANSession.GetValue();
+	}
+
+	if (NewSettings.JoinPolicy.IsSet())
+	{
+		SessionSettings.JoinPolicy = NewSettings.JoinPolicy.GetValue();
+
+		SetPermissionLevel(SessionModificationHandle, SessionSettings.JoinPolicy);
+	}
+
+	bool bMaxPlayersUpdated = false;
+
+	if (NewSettings.NumMaxPrivateConnections.IsSet())
+	{
+		SessionSettings.NumMaxPrivateConnections = NewSettings.NumMaxPrivateConnections.GetValue();
+		bMaxPlayersUpdated = true;
+	}
+
+	if (NewSettings.NumMaxPublicConnections.IsSet())
+	{
+		SessionSettings.NumMaxPublicConnections = NewSettings.NumMaxPublicConnections.GetValue();
+		bMaxPlayersUpdated = true;
+	}
+
+	if (bMaxPlayersUpdated)
+	{
+		SetMaxPlayers(SessionModificationHandle, SessionSettings.NumMaxPrivateConnections + SessionSettings.NumMaxPublicConnections);
+	}
+
+	if (NewSettings.NumOpenPrivateConnections.IsSet())
+	{
+		SessionSettings.NumOpenPrivateConnections = NewSettings.NumOpenPrivateConnections.GetValue();
+	}
+
+	if (NewSettings.NumOpenPublicConnections.IsSet())
+	{
+		SessionSettings.NumOpenPublicConnections = NewSettings.NumOpenPublicConnections.GetValue();
+	}
+
+	// TODO: We may need some additional logic for schema changes
+	if (NewSettings.SchemaName.IsSet())
+	{
+		SessionSettings.SchemaName = NewSettings.SchemaName.GetValue();
+	}
+
+	if (NewSettings.SessionIdOverride.IsSet())
+	{
+		SessionSettings.SessionIdOverride = NewSettings.SessionIdOverride.GetValue();
+	}
+
+	// Custom Settings
+
+	for (const FName& Key : NewSettings.RemovedCustomSettings)
+	{
+		SessionSettings.CustomSettings.Remove(Key);
+
+		// TODO: Remove the attribute to the SessionModification. Pending SchemaVariant work.
+	}
+
+// 	for (FCustomSessionSettingsMap::TConstIterator It(NewSettings.UpdatedCustomSettings); It; ++It)
+// 	{
+// 		// TODO: Add the attribute to the SessionModification. Pending SchemaVariant work.
+// 	}
+
+	SessionSettings.CustomSettings.Append(NewSettings.UpdatedCustomSettings);
+
+	// BucketId has its own set method on the API
+	const FCustomSessionSetting* NewBucketIdSetting = NewSettings.UpdatedCustomSettings.Find(EOS_SESSIONS_BUCKET_ID);
+	if (NewBucketIdSetting)
+	{
+		SetBucketId(SessionModificationHandle, NewBucketIdSetting->Data.Get<FString>());
+	}
+
+	// Registered Users
+
+	for (const FOnlineAccountIdHandle& Key : NewSettings.RemovedRegisteredUsers)
+	{
+		SessionSettings.RegisteredUsers.Remove(Key);
+	}
+
+	SessionSettings.RegisteredUsers.Append(NewSettings.UpdatedRegisteredUsers);
+
+	// Session Members
+
+	for (const FOnlineAccountIdHandle& Key : NewSettings.RemovedSessionMembers)
+	{
+// 		// Before removing a session member, we'll remove all their attributes
+// 		if (FSessionMember* SessionMember = SessionSettings.SessionMembers.Find(Key))
+// 		{
+// 			TArray<FName> SettingsKeyArray;
+// 			SessionMember->MemberSettings.GenerateKeyArray(SettingsKeyArray);
+// 
+// 			for (const FName& SettingKey : SettingsKeyArray)
+// 			{
+// 				const FName& MemberKey = FName(LexToString(GetProductUserIdChecked(Key)) + TEXT(":") + SettingKey.ToString());
+// 				
+// 				// TODO: Remove the attribute to the SessionModification. Pending SchemaVariant work.
+// 			}
+// 		}
+
+		SessionSettings.SessionMembers.Remove(Key);
+	}
+
+	for (FSessionMemberUpdatesMap::TConstIterator MemberIt(NewSettings.UpdatedSessionMembers); MemberIt; ++MemberIt)
+	{
+		if (FSessionMember* SessionMember = SessionSettings.SessionMembers.Find(MemberIt.Key()))
+		{
+			const FSessionMemberUpdate& SessionMemberUpdate = MemberIt.Value();
+
+			for (const FName& Key : SessionMemberUpdate.RemovedMemberSettings)
+			{
+				SessionMember->MemberSettings.Remove(Key);
+
+// 				const FName& MemberKey = FName(LexToString(GetProductUserIdChecked(MemberIt.Key())) + TEXT(":") + Key.ToString());
+// 				
+// 				// TODO: Remove the attribute to the SessionModification. Pending SchemaVariant work.
+			}
+			
+// 			for (FCustomSessionSettingsMap::TConstIterator SettingIt(SessionMemberUpdate.UpdatedMemberSettings); SettingIt; ++SettingIt)
+// 			{
+// 				const FName& Key = FName(LexToString(GetProductUserIdChecked(MemberIt.Key())) + TEXT(":") + SettingIt.Key().ToString());
+// 				const FCustomSessionSetting& Value = SettingIt.Value();
+// 
+// 				// TODO: Add the attribute to the SessionModification. Pending SchemaVariant work.
+// 			}
+
+			SessionMember->MemberSettings.Append(SessionMemberUpdate.UpdatedMemberSettings);
+		}
+	}
+}
+
+TOnlineAsyncOpHandle<FUpdateSession> FSessionsEOSGS::UpdateSession(FUpdateSession::Params&& Params)
+{
+	TOnlineAsyncOpRef<FUpdateSession> Op = GetOp<FUpdateSession>(MoveTemp(Params));
+	const FUpdateSession::Params& OpParams = Op->GetParams();
+
+	Op->Then([this](TOnlineAsyncOp<FUpdateSession>& Op) mutable
+	{
+		const FUpdateSession::Params& OpParams = Op.GetParams();
+
+		if (TSharedRef<FSessionEOSGS>* FoundSession = SessionsByName.Find(OpParams.SessionName))
+		{
+			if (FoundSession->Get().CurrentState == ESessionState::Creating)
+			{
+				Op.SetError(Errors::InvalidState());
+
+				return MakeFulfilledPromise<TDefaultErrorResult<FUpdateSessionImpl>>().GetFuture();
+			}
+
+			if (!FoundSession->Get().SessionSettings.IsDedicatedServerSession)
+			{
+				IAuthPtr Auth = Services.GetAuthInterface();
+
+				TOnlineResult<FAuthGetAccountByAccountId> Result = Auth->GetAccountByAccountId({ OpParams.LocalUserId });
+
+				if (Result.IsOk())
+				{
+					if (Result.GetOkValue().AccountInfo->LoginStatus == ELoginStatus::NotLoggedIn)
+					{
+						UE_LOG(LogTemp, Warning, TEXT("[FSessionsEOSGS::UpdateSession] Could not update session with user [%s] not logged in"), *ToLogString(OpParams.LocalUserId));
+
+						Op.SetError(Errors::InvalidUser());
+
+						return MakeFulfilledPromise<TDefaultErrorResult<FUpdateSessionImpl>>().GetFuture();
+					}
+				}
+			}
+
+			EOS_Sessions_UpdateSessionModificationOptions UpdateSessionModificationOptions = {};
+			UpdateSessionModificationOptions.ApiVersion = EOS_SESSIONS_UPDATESESSIONMODIFICATION_API_LATEST;
+			static_assert(EOS_SESSIONS_UPDATESESSIONMODIFICATION_API_LATEST == 1, "EOS_Sessions_UpdateSessionModificationOptions updated, check new fields");
+
+			FTCHARToUTF8 SessionNameUtf8(OpParams.SessionName.ToString());
+			UpdateSessionModificationOptions.SessionName = SessionNameUtf8.Get();
+
+			EOS_HSessionModification SessionModificationHandle = nullptr;
+			EOS_EResult ResultCode = EOS_Sessions_UpdateSessionModification(SessionsHandle, &UpdateSessionModificationOptions, &SessionModificationHandle);
+			if (ResultCode != EOS_EResult::EOS_Success)
+			{
+				UE_LOG(LogTemp, Error, TEXT("[FSessionsEOSGS::UpdateSession] EOS_Sessions_UpdateSessionModification failed with result [%s]"), *LexToString(ResultCode));
+
+				Op.SetError(FromEOSError(ResultCode));
+
+				return MakeFulfilledPromise<TDefaultErrorResult<FUpdateSessionImpl>>().GetFuture();
+			}
+
+			// After creating the session modification handle, we'll add all the updated data to it
+			WriteUpdateSessionModificationHandle(SessionModificationHandle, FoundSession->Get().SessionSettings, OpParams.Mutations);
+
+			FUpdateSessionImpl::Params UpdateSessionImplParams{ SessionModificationHandle };
+
+			return UpdateSessionImpl(MoveTemp(UpdateSessionImplParams));
+		}
+		else
+		{
+			Op.SetError(Errors::NotFound());
+
+			return MakeFulfilledPromise<TDefaultErrorResult<FUpdateSessionImpl>>().GetFuture();
+		}
+	})
+	.Then([this](TOnlineAsyncOp<FUpdateSession>& Op, TDefaultErrorResult<FUpdateSessionImpl>&& Result) mutable
+	{
+		if (Result.IsOk())
+		{
+			const FUpdateSession::Params& OpParams = Op.GetParams();
+
+			if (TSharedRef<FSessionEOSGS>* Session = SessionsByName.Find(OpParams.SessionName))
+			{
+				// We set the result and fire the event
+				Op.SetResult(FUpdateSession::Result{ *Session });
+
+				FSessionUpdated SessionUpdatedEvent{ *Session, OpParams.Mutations };
+				SessionEvents.OnSessionUpdated.Broadcast(SessionUpdatedEvent);
+			}
+			else
+			{
+				Op.SetError(Errors::NotFound());
+			}
+		}
+		else
+		{
+			Op.SetError(MoveTemp(Result.GetErrorValue()));
+		}
+	})
+	.Enqueue(GetSerialQueue());
+
+	return Op->GetHandle();
+}
+
+TFuture<TDefaultErrorResult<FUpdateSessionImpl>> FSessionsEOSGS::UpdateSessionImpl(FUpdateSessionImpl::Params&& Params)
+{
+	EOS_Sessions_UpdateSessionOptions UpdateSessionOptions = {};
+	UpdateSessionOptions.ApiVersion = EOS_SESSIONS_UPDATESESSION_API_LATEST;
+	static_assert(EOS_SESSIONS_UPDATESESSION_API_LATEST == 1, "EOS_Sessions_UpdateSessionOptions updated, check new fields");
+
+	UpdateSessionOptions.SessionModificationHandle = Params.SessionModificationHandle;
+
+	TPromise<TDefaultErrorResult<FUpdateSessionImpl>> Promise;
+	TFuture<TDefaultErrorResult<FUpdateSessionImpl>> Future = Promise.GetFuture();
+
+	EOS_Async<EOS_Sessions_UpdateSessionCallbackInfo>(EOS_Sessions_UpdateSession, SessionsHandle, UpdateSessionOptions)
+		.Then([this, Promise = MoveTemp(Promise)](TFuture<const EOS_Sessions_UpdateSessionCallbackInfo*>&& Future) mutable
+	{
+		const EOS_Sessions_UpdateSessionCallbackInfo* Result = Future.Get();
+		if (Result->ResultCode != EOS_EResult::EOS_Success && Result->ResultCode != EOS_EResult::EOS_Sessions_OutOfSync)
+		{
+			Promise.EmplaceValue(FromEOSError(Result->ResultCode));
+			return;
+		}
+
+		if (TSharedRef<FSessionEOSGS>* FoundSession = SessionsByName.Find(FName(Result->SessionName)))
+		{
+			// In session creation calls, we'll need to set the new id for the session
+			if (!FoundSession->Get().SessionId.IsValid())
+			{
+				FoundSession->Get().SessionId = FOnlineSessionIdRegistryEOSGS::Get().FindOrAdd(FString(Result->SessionId));
+			}
+		}
+
+		Promise.EmplaceValue(FUpdateSessionImpl::Result { });
+	});
+
+	return Future;
+}
+
+TOnlineAsyncOpHandle<FLeaveSession> FSessionsEOSGS::LeaveSession(FLeaveSession::Params&& Params)
+{
+	TOnlineAsyncOpRef<FLeaveSession> Op = GetOp<FLeaveSession>(MoveTemp(Params));
+
+	Op->Then([this](TOnlineAsyncOp<FLeaveSession>& Op) mutable
+	{
+		const FLeaveSession::Params& OpParams = Op.GetParams();
+
+		// We do not check if OpParams.LocalUserId is logged in because it might be a dedicated server destroying the session
+
+		if (TSharedRef<FSessionEOSGS>* FoundSession = SessionsByName.Find(OpParams.SessionName))
+		{
+			EOS_Sessions_DestroySessionOptions DestroySessionOptions = {};
+			DestroySessionOptions.ApiVersion = EOS_SESSIONS_DESTROYSESSION_API_LATEST;
+			static_assert(EOS_SESSIONS_DESTROYSESSION_API_LATEST == 1, "EOS_Sessions_DestroySessionOptions updated, check new fields");
+
+			FTCHARToUTF8 SessionNameUtf8(OpParams.SessionName.ToString());
+			DestroySessionOptions.SessionName = SessionNameUtf8.Get();
+
+			return EOS_Async<EOS_Sessions_DestroySessionCallbackInfo>(EOS_Sessions_DestroySession, SessionsHandle, DestroySessionOptions);
+		}
+		else
+		{
+			Op.SetError(Errors::NotFound());
+
+			return MakeFulfilledPromise<const EOS_Sessions_DestroySessionCallbackInfo*>().GetFuture();
+		}
+	})
+	.Then([this](TOnlineAsyncOp<FLeaveSession>& Op, const EOS_Sessions_DestroySessionCallbackInfo* Result) mutable
+	{
+		if (Result->ResultCode != EOS_EResult::EOS_Success)
+		{
+			Op.SetError(FromEOSError(Result->ResultCode));
+			return;
+		}
+
+		const FLeaveSession::Params& OpParams = Op.GetParams();
+
+		if (TSharedRef<FSessionEOSGS>* FoundSession = SessionsByName.Find(OpParams.SessionName))
+		{
+			SessionsById.Remove((*FoundSession)->SessionId);
+			SessionsByName.Remove(OpParams.SessionName);
+
+			Op.SetResult(FLeaveSession::Result{ });
+
+			FSessionLeft SessionLeftEvent;
+			SessionLeftEvent.LocalUserIds.Add(OpParams.LocalUserId);
+			//SessionLeftEvent.LocalUserIds.Append(OpParams.AdditionalLocalUsers); // TODO: Process multiple users in Session destruction
+			SessionEvents.OnSessionLeft.Broadcast(SessionLeftEvent);
+		}
+		else
+		{
+			Op.SetError(Errors::NotFound());
+		}
+	})
+	.Enqueue(GetSerialQueue());
+
+	return Op->GetHandle();
+}
+
+/* UE::Online */ }
\ No newline at end of file
diff --git a/Engine/Plugins/Online/Experimental/OnlineServicesEOSGS/Source/Public/Online/SessionsEOSGS.h b/Engine/Plugins/Online/Experimental/OnlineServicesEOSGS/Source/Public/Online/SessionsEOSGS.h
new file mode 100644
--- /dev/null
+++ b/Engine/Plugins/Online/Experimental/OnlineServicesEOSGS/Source/Public/Online/SessionsEOSGS.h
@@ -0,0 +1,122 @@
+// Copyright Epic Games, Inc. All Rights Reserved.
+
+#pragma once
+
+#include "CoreMinimal.h"
+#include "Online/SessionsCommon.h"
+#include "IPAddress.h"
+
+#if defined(EOS_PLATFORM_BASE_FILE_NAME)
+#include EOS_PLATFORM_BASE_FILE_NAME
+#endif
+#include "eos_sessions_types.h"
+
+namespace UE::Online {
+
+class FOnlineServicesEOSGS;
+
+/** Session Id Registry */
+
+// TODO: Move to Common since it's in Null too
+class FOnlineSessionIdString
+{
+public:
+	FString Data;
+	FOnlineSessionIdHandle Handle;
+};
+
+class FOnlineSessionIdRegistryEOSGS : public IOnlineSessionIdRegistry
+{
+public:
+	static FOnlineSessionIdRegistryEOSGS& Get();
+
+	FOnlineSessionIdHandle Find(const FString& SessionId) const;
+	FOnlineSessionIdHandle FindOrAdd(const FString& SessionId);
+
+	// Begin IOnlineSessionIdRegistry
+	virtual FString ToLogString(const FOnlineSessionIdHandle& Handle) const override;
+	virtual TArray<uint8> ToReplicationData(const FOnlineSessionIdHandle& Handle) const override;
+	virtual FOnlineSessionIdHandle FromReplicationData(const TArray<uint8>& ReplicationString) override;
+	// End IOnlineSessionIdRegistry
+
+	virtual ~FOnlineSessionIdRegistryEOSGS() = default;
+
+private:
+	const FOnlineSessionIdString* GetInternal(const FOnlineSessionIdHandle& Handle) const;
+
+	TArray<FOnlineSessionIdString> Ids;
+	TMap<FString, FOnlineSessionIdString> StringToId;
+};
+
+static FName EOS_SESSIONS_BUCKET_ID = TEXT("EOS_SESSIONS_BUCKET_ID");
+
+class FSessionEOSGS : public FSession
+{
+public:
+	FSessionEOSGS();
+	FSessionEOSGS(const FSession& InSession);
+
+public:
+	/** The IP address for the session owner */
+	TSharedPtr<FInternetAddr> OwnerInternetAddr;
+};
+
+struct FUpdateSessionImpl
+{
+	static constexpr TCHAR Name[] = TEXT("FUpdateSessionImpl");
+
+	struct Params
+	{
+		/** Handle for the session modification operation */
+		EOS_HSessionModification SessionModificationHandle;
+	};
+
+	struct Result
+	{
+	};
+};
+
+class ONLINESERVICESEOSGS_API FSessionsEOSGS : public FSessionsCommon
+{
+public:
+	using Super = FSessionsCommon;
+
+	FSessionsEOSGS(FOnlineServicesEOSGS& InOwningSubsystem);
+	virtual ~FSessionsEOSGS() = default;
+
+	// TOnlineComponent
+	void Initialize() override;
+
+	// ISessions
+	virtual TOnlineResult<FGetAllSessions> GetAllSessions(FGetAllSessions::Params&& Params) override;
+	virtual TOnlineResult<FGetSessionByName> GetSessionByName(FGetSessionByName::Params&& Params) override;
+	virtual TOnlineResult<FGetSessionById> GetSessionById(FGetSessionById::Params&& Params) override;
+	virtual TOnlineAsyncOpHandle<FCreateSession> CreateSession(FCreateSession::Params&& Params) override;
+	virtual TOnlineAsyncOpHandle<FUpdateSession> UpdateSession(FUpdateSession::Params&& Params) override;
+	virtual TOnlineAsyncOpHandle<FLeaveSession> LeaveSession(FLeaveSession::Params&& Params) override;
+
+protected:
+	void SetPermissionLevel(EOS_HSessionModification& SessionModHandle, const ESessionJoinPolicy& JoinPolicy);
+	void SetBucketId(EOS_HSessionModification& SessionModHandle, const FString& NewBucketId);
+	void SetMaxPlayers(EOS_HSessionModification& SessionModHandle, const int32& NewMaxPlayers);
+
+	/**
+	 * Writes all values in the passed SessionSettings to the SessionModificationHandle
+	 */
+	void WriteCreateSessionModificationHandle(EOS_HSessionModification& SessionModificationHandle, const FSessionSettings& SessionSettings);
+	/**
+	 * Writes only the new values for all updated settings to the SessionModificationHandle, as well as updating them on the passed SessionSettings
+	 */
+	void WriteUpdateSessionModificationHandle(EOS_HSessionModification& SessionModificationHandle, FSessionSettings& SessionSettings, const FSessionSettingsUpdate& NewSettings);
+	TFuture<TDefaultErrorResult<FUpdateSessionImpl>> UpdateSessionImpl(FUpdateSessionImpl::Params&& Params);
+
+protected:
+	FOnlineServicesEOSGS& Services;
+
+	EOS_HSessions SessionsHandle = nullptr;
+
+	TMap<FName, TSharedRef<FSessionEOSGS>> SessionsByName;
+	TMap<FOnlineSessionIdHandle, TSharedRef<FSessionEOSGS>> SessionsById;
+};
+
+/* UE::Online */ }
\ No newline at end of file
