Create a default EOS platform on startup, and have the default instance use it

diff --git a/Engine/Plugins/Online/Experimental/OnlineServicesEOSGS/Source/Private/Online/OnlineServicesEOSGS.cpp b/Engine/Plugins/Online/Experimental/OnlineServicesEOSGS/Source/Private/Online/OnlineServicesEOSGS.cpp
--- a/Engine/Plugins/Online/Experimental/OnlineServicesEOSGS/Source/Private/Online/OnlineServicesEOSGS.cpp
+++ b/Engine/Plugins/Online/Experimental/OnlineServicesEOSGS/Source/Private/Online/OnlineServicesEOSGS.cpp
@@ -9,10 +9,7 @@
 #include "Online/SessionsEOSGS.h"
 #include "Online/OnlineIdEOSGS.h"
 #include "Online/OnlineServicesEOSGSTypes.h"
-
-#include "EOSShared.h"
-#include "EOSSharedTypes.h"
-#include "IEOSSDKManager.h"
+#include "Online/OnlineServicesEOSGSPlatformFactory.h"
 
 #if WITH_ENGINE
 #include "InternetAddrEOS.h"
@@ -22,29 +19,8 @@
 
 namespace UE::Online {
 
-struct FEOSPlatformConfig
-{
-	FString ProductId;
-	FString SandboxId;
-	FString DeploymentId;
-	FString ClientId;
-	FString ClientSecret;
-};
-
-namespace Meta {
-
-BEGIN_ONLINE_STRUCT_META(FEOSPlatformConfig)
-	ONLINE_STRUCT_FIELD(FEOSPlatformConfig, ProductId),
-	ONLINE_STRUCT_FIELD(FEOSPlatformConfig, SandboxId),
-	ONLINE_STRUCT_FIELD(FEOSPlatformConfig, DeploymentId),
-	ONLINE_STRUCT_FIELD(FEOSPlatformConfig, ClientId),
-	ONLINE_STRUCT_FIELD(FEOSPlatformConfig, ClientSecret)
-END_ONLINE_STRUCT_META()
-
-/* Meta */ }
-
 FOnlineServicesEOSGS::FOnlineServicesEOSGS(FName InInstanceName)
-	: Super(TEXT("EOS"), InInstanceName)
+	: Super(GetConfigNameStatic(), InInstanceName)
 {
 }
 
@@ -60,48 +36,19 @@ void FOnlineServicesEOSGS::RegisterComponents()
 
 void FOnlineServicesEOSGS::Initialize()
 {
-	IEOSSDKManager* SDKManager = IEOSSDKManager::Get();
-	if (!SDKManager)
+	FOnlineServicesEOSGSPlatformFactory& PlatformFactory = FOnlineServicesEOSGSPlatformFactory::Get();
+	if (InstanceName.IsNone())
 	{
-		return;
+		EOSPlatformHandle = PlatformFactory.GetDefaultPlatform();
 	}
-
-	EOS_EResult InitResult = SDKManager->Initialize();
-	if (InitResult != EOS_EResult::EOS_Success)
+	else
 	{
-		return;
+		EOSPlatformHandle = PlatformFactory.CreatePlatform();
 	}
-
-	FEOSPlatformConfig EOSPlatformConfig;
-	LoadConfig(EOSPlatformConfig);
-
-	const FTCHARToUTF8 ProductId(*EOSPlatformConfig.ProductId);
-	const FTCHARToUTF8 SandboxId(*EOSPlatformConfig.SandboxId);
-	const FTCHARToUTF8 DeploymentId(*EOSPlatformConfig.DeploymentId);
-	const FTCHARToUTF8 ClientId(*EOSPlatformConfig.ClientId);
-	const FTCHARToUTF8 ClientSecret(*EOSPlatformConfig.ClientSecret);
-	const FTCHARToUTF8 CacheDirectory(*(SDKManager->GetCacheDirBase() / TEXT("OnlineServicesEOS")));
-
-	EOS_Platform_Options PlatformOptions = {};
-	PlatformOptions.ApiVersion = EOS_PLATFORM_OPTIONS_API_LATEST;
-	PlatformOptions.Reserved = nullptr;
-	PlatformOptions.bIsServer = EOS_FALSE;
-	PlatformOptions.OverrideCountryCode = nullptr;
-	PlatformOptions.OverrideLocaleCode = nullptr;
-	PlatformOptions.Flags = EOS_PF_WINDOWS_ENABLE_OVERLAY_D3D9 | EOS_PF_WINDOWS_ENABLE_OVERLAY_D3D10 | EOS_PF_WINDOWS_ENABLE_OVERLAY_OPENGL; // Enable overlay support for D3D9/10 and OpenGL. This sample uses D3D11 or SDL.
-
-	PlatformOptions.ProductId = ProductId.Get();
-	PlatformOptions.SandboxId = SandboxId.Get();
-	PlatformOptions.DeploymentId = DeploymentId.Get();
-	PlatformOptions.ClientCredentials.ClientId = ClientId.Get();
-	PlatformOptions.ClientCredentials.ClientSecret = ClientSecret.Get();
-	PlatformOptions.CacheDirectory = CacheDirectory.Get();
-
-	EOSPlatformHandle = SDKManager->CreatePlatform(PlatformOptions);
 	if (EOSPlatformHandle)
 	{
 #if WITH_ENGINE
-		SocketSubsystem = MakeShareable(new FSocketSubsystemEOS(EOSPlatformHandle, MakeShareable(new FSocketSubsystemEOSUtils_OnlineServicesEOS(*this))));
+		SocketSubsystem = MakeShared<FSocketSubsystemEOS>(EOSPlatformHandle, MakeShared< FSocketSubsystemEOSUtils_OnlineServicesEOS>(*this));
 		check(SocketSubsystem);
 
 		FString ErrorStr;
diff --git a/Engine/Plugins/Online/Experimental/OnlineServicesEOSGS/Source/Private/Online/OnlineServicesEOSGSPlatformFactory.cpp b/Engine/Plugins/Online/Experimental/OnlineServicesEOSGS/Source/Private/Online/OnlineServicesEOSGSPlatformFactory.cpp
new file mode 100644
--- /dev/null
+++ b/Engine/Plugins/Online/Experimental/OnlineServicesEOSGS/Source/Private/Online/OnlineServicesEOSGSPlatformFactory.cpp
@@ -0,0 +1,115 @@
+// Copyright Epic Games, Inc. All Rights Reserved.
+
+#include "Online/OnlineServicesEOSGSPlatformFactory.h"
+
+#include "Online/OnlineServicesEOSGS.h"
+
+#include "Misc/ConfigCacheIni.h"
+#include "Misc/LazySingleton.h"
+
+#include "EOSShared.h"
+#include "EOSSharedTypes.h"
+#include "IEOSSDKManager.h"
+
+#include "CoreMinimal.h"
+
+struct FEOSPlatformConfig
+{
+	FString ProductId;
+	FString SandboxId;
+	FString DeploymentId;
+	FString ClientId;
+	FString ClientSecret;
+};
+
+FEOSPlatformConfig LoadEOSPlatformConfig()
+{
+	FEOSPlatformConfig PlatformConfig;
+	FString ConfigSection = FString::Printf(TEXT("OnlineServices.%s"), UE::Online::FOnlineServicesEOSGS::GetConfigNameStatic());
+	GConfig->GetString(*ConfigSection, TEXT("ProductId"), PlatformConfig.ProductId, GEngineIni);
+	GConfig->GetString(*ConfigSection, TEXT("SandboxId"), PlatformConfig.SandboxId, GEngineIni);
+	GConfig->GetString(*ConfigSection, TEXT("DeploymentId"), PlatformConfig.DeploymentId, GEngineIni);
+	GConfig->GetString(*ConfigSection, TEXT("ClientId"), PlatformConfig.ClientId, GEngineIni);
+	GConfig->GetString(*ConfigSection, TEXT("ClientSecret"), PlatformConfig.ClientSecret, GEngineIni);
+	return PlatformConfig;
+}
+
+namespace UE::Online {
+
+FOnlineServicesEOSGSPlatformFactory::FOnlineServicesEOSGSPlatformFactory()
+{
+	DefaultEOSPlatformHandle = CreatePlatform();
+}
+
+FOnlineServicesEOSGSPlatformFactory& FOnlineServicesEOSGSPlatformFactory::Get()
+{
+	return TLazySingleton<FOnlineServicesEOSGSPlatformFactory>::Get();
+}
+
+void FOnlineServicesEOSGSPlatformFactory::TearDown()
+{
+	return TLazySingleton<FOnlineServicesEOSGSPlatformFactory>::TearDown();
+}
+
+IEOSPlatformHandlePtr FOnlineServicesEOSGSPlatformFactory::CreatePlatform()
+{
+	const FName EOSSharedModuleName = TEXT("EOSShared");
+	if (!FModuleManager::Get().IsModuleLoaded(EOSSharedModuleName))
+	{
+		FModuleManager::Get().LoadModuleChecked(EOSSharedModuleName);
+	}
+	IEOSSDKManager* const SDKManager = IEOSSDKManager::Get();
+	if (!SDKManager)
+	{
+		return {};
+	}
+
+	EOS_EResult InitResult = SDKManager->Initialize();
+	if (InitResult != EOS_EResult::EOS_Success)
+	{
+		return {};
+	}
+
+	// Load config
+	FEOSPlatformConfig EOSPlatformConfig = LoadEOSPlatformConfig();
+	const FTCHARToUTF8 ProductId(*EOSPlatformConfig.ProductId);
+	const FTCHARToUTF8 SandboxId(*EOSPlatformConfig.SandboxId);
+	const FTCHARToUTF8 DeploymentId(*EOSPlatformConfig.DeploymentId);
+	const FTCHARToUTF8 ClientId(*EOSPlatformConfig.ClientId);
+	const FTCHARToUTF8 ClientSecret(*EOSPlatformConfig.ClientSecret);
+	const FTCHARToUTF8 CacheDirectory(*(SDKManager->GetCacheDirBase() / TEXT("OnlineServicesEOS")));
+
+	EOS_Platform_Options PlatformOptions = {};
+	PlatformOptions.ApiVersion = EOS_PLATFORM_OPTIONS_API_LATEST;
+	PlatformOptions.Reserved = nullptr;
+	PlatformOptions.bIsServer = EOS_FALSE;
+	PlatformOptions.OverrideCountryCode = nullptr;
+	PlatformOptions.OverrideLocaleCode = nullptr;
+	// Can't check GIsEditor here because it is too soon!
+	if (!IsRunningGame())
+	{
+		PlatformOptions.Flags = EOS_PF_LOADING_IN_EDITOR;
+	}
+	else
+	{
+		PlatformOptions.Flags = EOS_PF_WINDOWS_ENABLE_OVERLAY_D3D9 | EOS_PF_WINDOWS_ENABLE_OVERLAY_D3D10 | EOS_PF_WINDOWS_ENABLE_OVERLAY_OPENGL; // Enable overlay support for D3D9/10 and OpenGL. This sample uses D3D11 or SDL.
+	}
+
+	PlatformOptions.ProductId = ProductId.Get();
+	PlatformOptions.SandboxId = SandboxId.Get();
+	PlatformOptions.DeploymentId = DeploymentId.Get();
+	PlatformOptions.ClientCredentials.ClientId = ClientId.Get();
+	PlatformOptions.ClientCredentials.ClientSecret = ClientSecret.Get();
+	PlatformOptions.CacheDirectory = CacheDirectory.Get();
+
+	IEOSPlatformHandlePtr EOSPlatformHandle = SDKManager->CreatePlatform(PlatformOptions);
+	if (!EOSPlatformHandle)
+	{
+		UE_LOG(LogOnlineServices, Verbose, TEXT("[FOnlineServicesEOSGS::Initialize] Unable to initialize Socket Subsystem. EOS Platform Handle was invalid."));
+		return {};
+	}
+	return EOSPlatformHandle;
+}
+
+
+/* UE::Online */ }
diff --git a/Engine/Plugins/Online/Experimental/OnlineServicesEOSGS/Source/Private/Online/OnlineServicesEOSGSPlatformFactory.h b/Engine/Plugins/Online/Experimental/OnlineServicesEOSGS/Source/Private/Online/OnlineServicesEOSGSPlatformFactory.h
new file mode 100644
--- /dev/null
+++ b/Engine/Plugins/Online/Experimental/OnlineServicesEOSGS/Source/Private/Online/OnlineServicesEOSGSPlatformFactory.h
@@ -0,0 +1,45 @@
+// Copyright Epic Games, Inc. All Rights Reserved.
+
+#pragma once
+
+#include "CoreMinimal.h"
+
+using IEOSPlatformHandlePtr = TSharedPtr<class IEOSPlatformHandle>;
+
+namespace UE::Online
+{
+
+/** Factory class to create EOS Platforms for online services */
+class FOnlineServicesEOSGSPlatformFactory
+{
+public:
+	/**
+	 * Get the platform factory
+	 * @return the platform factory singleton
+	 */
+	static FOnlineServicesEOSGSPlatformFactory& Get();
+	/**
+	 * Tear down the singleton instance. This only cleans up the singleton and has no impact on any platform handles created by this (aside from DefaultEOSPlatformHandle's ref count decreasing).
+	 */
+	static void TearDown();
+
+	/**
+	 * Create a new platform. Loads configuration from ini.
+	 * @return a new platform, or null on failure
+	 */
+	IEOSPlatformHandlePtr CreatePlatform();
+	/**
+	 * Get the default platform.
+	 * The default platform is created on startup and can be used by the default instance of online services EOS.
+	 * @return the default platform
+	 */
+	IEOSPlatformHandlePtr GetDefaultPlatform() { return DefaultEOSPlatformHandle; }
+private:
+	/** Default constructor */
+	FOnlineServicesEOSGSPlatformFactory();
+	friend FLazySingleton;
+	/** Platform handle for the default instance */
+	IEOSPlatformHandlePtr DefaultEOSPlatformHandle;
+};
+
+/* UE::Online */ }
diff --git a/Engine/Plugins/Online/Experimental/OnlineServicesEOSGS/Source/Private/OnlineServicesEOSGSModule.cpp b/Engine/Plugins/Online/Experimental/OnlineServicesEOSGS/Source/Private/OnlineServicesEOSGSModule.cpp
--- a/Engine/Plugins/Online/Experimental/OnlineServicesEOSGS/Source/Private/OnlineServicesEOSGSModule.cpp
+++ b/Engine/Plugins/Online/Experimental/OnlineServicesEOSGS/Source/Private/OnlineServicesEOSGSModule.cpp
@@ -5,6 +5,7 @@
 #include "Online/OnlineIdEOSGS.h"
 #include "Online/OnlineServicesEOSGS.h"
 #include "Online/OnlineServicesRegistry.h"
+#include "Online/OnlineServicesEOSGSPlatformFactory.h"
 
 #include "CoreMinimal.h"
 
@@ -15,6 +16,7 @@ class FOnlineServicesEOSGSModule : public IModuleInterface
 {
 public:
 	virtual void StartupModule() override;
+	virtual void ShutdownModule() override;
 protected:
 };
 
@@ -33,6 +35,14 @@ void FOnlineServicesEOSGSModule::StartupModule()
 {
 	FOnlineServicesRegistry::Get().RegisterServicesFactory(EOnlineServices::Epic, MakeUnique<FOnlineServicesFactoryEOSGS>());
 	FOnlineIdRegistryRegistry::Get().RegisterAccountIdRegistry(EOnlineServices::Epic, &FOnlineAccountIdRegistryEOSGS::Get());
+
+	// Initialize the platform factory on startup.  This is necessary for the SDK to bind to rendering and input very early.
+	FOnlineServicesEOSGSPlatformFactory::Get();
+}
+
+void FOnlineServicesEOSGSModule::ShutdownModule()
+{
+	UE::Online::FOnlineServicesEOSGSPlatformFactory::TearDown();
 }
 
 /* UE::Online */ }
diff --git a/Engine/Plugins/Online/Experimental/OnlineServicesEOSGS/Source/Public/Online/OnlineServicesEOSGS.h b/Engine/Plugins/Online/Experimental/OnlineServicesEOSGS/Source/Public/Online/OnlineServicesEOSGS.h
--- a/Engine/Plugins/Online/Experimental/OnlineServicesEOSGS/Source/Public/Online/OnlineServicesEOSGS.h
+++ b/Engine/Plugins/Online/Experimental/OnlineServicesEOSGS/Source/Public/Online/OnlineServicesEOSGS.h
@@ -31,6 +31,8 @@ public:
 	virtual EOnlineServices GetServicesProvider() const override { return EOnlineServices::Epic; }
 
 	EOS_HPlatform GetEOSPlatformHandle() const;
+
+	static const TCHAR* GetConfigNameStatic() { return TEXT("EOS"); }
 protected:
 	IEOSPlatformHandlePtr EOSPlatformHandle;
 
