Check EOS config validity before calling EOS SDK functions
Fixes CIS warnings on module startup due to the plugin being enabled, but not used.

diff --git a/Engine/Plugins/Online/Experimental/OnlineServicesEOSGS/Source/Private/Online/OnlineServicesEOSGSPlatformFactory.cpp b/Engine/Plugins/Online/Experimental/OnlineServicesEOSGS/Source/Private/Online/OnlineServicesEOSGSPlatformFactory.cpp
--- a/Engine/Plugins/Online/Experimental/OnlineServicesEOSGS/Source/Private/Online/OnlineServicesEOSGSPlatformFactory.cpp
+++ b/Engine/Plugins/Online/Experimental/OnlineServicesEOSGS/Source/Private/Online/OnlineServicesEOSGSPlatformFactory.cpp
@@ -6,6 +6,7 @@
 
 #include "Misc/ConfigCacheIni.h"
 #include "Misc/LazySingleton.h"
+#include "Modules/ModuleManager.h"
 
 #include "EOSShared.h"
 #include "EOSSharedTypes.h"
@@ -34,6 +35,15 @@ FEOSPlatformConfig LoadEOSPlatformConfig()
 	return PlatformConfig;
 }
 
+bool IsEOSPlatformConfigValid(const FEOSPlatformConfig& InConfig)
+{
+	return !InConfig.ProductId.IsEmpty() &&
+		!InConfig.SandboxId.IsEmpty() &&
+		!InConfig.DeploymentId.IsEmpty() &&
+		!InConfig.ClientId.IsEmpty() &&
+		!InConfig.ClientSecret.IsEmpty();
+}
+
 namespace UE::Online {
 
 FOnlineServicesEOSGSPlatformFactory::FOnlineServicesEOSGSPlatformFactory()
@@ -72,6 +82,10 @@ IEOSPlatformHandlePtr FOnlineServicesEOSGSPlatformFactory::CreatePlatform()
 
 	// Load config
 	FEOSPlatformConfig EOSPlatformConfig = LoadEOSPlatformConfig();
+	if (!IsEOSPlatformConfigValid(EOSPlatformConfig))
+	{
+		return {};
+	}
 	const FTCHARToUTF8 ProductId(*EOSPlatformConfig.ProductId);
 	const FTCHARToUTF8 SandboxId(*EOSPlatformConfig.SandboxId);
 	const FTCHARToUTF8 DeploymentId(*EOSPlatformConfig.DeploymentId);
