- Fix some methods being attributed to IAuth which are really from IOnlineComponent

diff --git a/Engine/Plugins/Online/Experimental/OnlineServices/Source/OnlineServicesCommon/Public/Online/AuthCommon.h b/Engine/Plugins/Online/Experimental/OnlineServices/Source/OnlineServicesCommon/Public/Online/AuthCommon.h
--- a/Engine/Plugins/Online/Experimental/OnlineServices/Source/OnlineServicesCommon/Public/Online/AuthCommon.h
+++ b/Engine/Plugins/Online/Experimental/OnlineServices/Source/OnlineServicesCommon/Public/Online/AuthCommon.h
@@ -16,7 +16,7 @@ public:
 
 	FAuthCommon(FOnlineServicesCommon& InServices);
 
-	// TOnlineComponent
+	// IOnlineComponent
 	virtual void RegisterCommands() override;
 
 	// IAuth
diff --git a/Engine/Plugins/Online/Experimental/OnlineServicesEOS/Source/Public/Online/AuthEOS.h b/Engine/Plugins/Online/Experimental/OnlineServicesEOS/Source/Public/Online/AuthEOS.h
--- a/Engine/Plugins/Online/Experimental/OnlineServicesEOS/Source/Public/Online/AuthEOS.h
+++ b/Engine/Plugins/Online/Experimental/OnlineServicesEOS/Source/Public/Online/AuthEOS.h
@@ -16,8 +16,11 @@ public:
 	FAuthEOS(FOnlineServicesEOS& InOwningSubsystem);
 	virtual ~FAuthEOS() = default;
 
-	// Begin IAuth
+	// Begin IOnlineComponent
 	virtual void Initialize() override;
+	// End IOnlineComponent
+
+	// Begin IAuth
 	virtual TOnlineAsyncOpHandle<FAuthLogin> Login(FAuthLogin::Params&& Params) override;
 	virtual TOnlineAsyncOpHandle<FAuthLogout> Logout(FAuthLogout::Params&& Params) override;
 	// End IAuth
diff --git a/Engine/Plugins/Online/Experimental/OnlineServicesEOSGS/Source/Public/Online/AuthEOSGS.h b/Engine/Plugins/Online/Experimental/OnlineServicesEOSGS/Source/Public/Online/AuthEOSGS.h
--- a/Engine/Plugins/Online/Experimental/OnlineServicesEOSGS/Source/Public/Online/AuthEOSGS.h
+++ b/Engine/Plugins/Online/Experimental/OnlineServicesEOSGS/Source/Public/Online/AuthEOSGS.h
@@ -61,8 +61,11 @@ public:
 	FAuthEOSGS(FOnlineServicesEOSGS& InOwningSubsystem);
 	virtual ~FAuthEOSGS() = default;
 
-	// Begin IAuth
+	// Begin IOnlineComponent
 	virtual void Initialize() override;
+	// End IOnlineComponent
+
+	// Begin IAuth
 	virtual TOnlineAsyncOpHandle<FAuthLogin> Login(FAuthLogin::Params&& Params) override;
 	virtual TOnlineAsyncOpHandle<FAuthLogout> Logout(FAuthLogout::Params&& Params) override;
 	virtual TOnlineResult<FAuthGetAccountByPlatformUserId> GetAccountByPlatformUserId(FAuthGetAccountByPlatformUserId::Params&& Params) override;
