- Title File EOS Implementation

diff --git a/Engine/Plugins/Online/EOSShared/Source/EOSShared/Public/EOSSharedTypes.h b/Engine/Plugins/Online/EOSShared/Source/EOSShared/Public/EOSSharedTypes.h
--- a/Engine/Plugins/Online/EOSShared/Source/EOSShared/Public/EOSSharedTypes.h
+++ b/Engine/Plugins/Online/EOSShared/Source/EOSShared/Public/EOSSharedTypes.h
@@ -27,12 +27,12 @@ public:
 #if WITH_EOS_SDK
 
 /** Class to handle all callbacks generically using a lambda to process callback results */
-template<typename CallbackFuncType, typename CallbackType, typename OwningType>
+template<typename CallbackFuncType, typename CallbackParamType, typename OwningType, typename CallbackReturnType = void>
 class TEOSGlobalCallback :
 	public FCallbackBase
 {
 public:
-	TFunction<void(const CallbackType*)> CallbackLambda;
+	TFunction<CallbackReturnType(const CallbackParamType*)> CallbackLambda;
 
 	TEOSGlobalCallback(TWeakPtr<OwningType> InOwner)
 		: FCallbackBase()
@@ -41,7 +41,6 @@ public:
 	}
 	virtual ~TEOSGlobalCallback() = default;
 
-
 	CallbackFuncType GetCallbackPtr()
 	{
 		return &CallbackImpl;
@@ -54,7 +53,7 @@ private:
 	/** The object that needs to be checked for lifetime before calling the callback */
 	TWeakPtr<OwningType> Owner;
 
-	static void EOS_CALL CallbackImpl(const CallbackType* Data)
+	static CallbackReturnType EOS_CALL CallbackImpl(const CallbackParamType* Data)
 	{
 		TEOSGlobalCallback* CallbackThis = (TEOSGlobalCallback*)Data->ClientData;
 		check(CallbackThis);
@@ -67,7 +66,21 @@ private:
 		if (CallbackThis->Owner.IsValid())
 		{
 			check(CallbackThis->CallbackLambda);
-			CallbackThis->CallbackLambda(Data);
+
+			if constexpr (std::is_void<CallbackReturnType>::value)
+			{
+				CallbackThis->CallbackLambda(Data);
+			}
+			else
+			{
+				return CallbackThis->CallbackLambda(Data);
+			}
+		}
+		
+		if constexpr (!std::is_void<CallbackReturnType>::value)
+		{
+			// we need to return _something_ to compile.
+			return CallbackReturnType{};
 		}
 	}
 };
diff --git a/Engine/Plugins/Online/Experimental/OnlineServicesEOSGS/Source/Private/Online/TitleFileEOSGS.cpp b/Engine/Plugins/Online/Experimental/OnlineServicesEOSGS/Source/Private/Online/TitleFileEOSGS.cpp
new file mode 100644
--- /dev/null
+++ b/Engine/Plugins/Online/Experimental/OnlineServicesEOSGS/Source/Private/Online/TitleFileEOSGS.cpp
@@ -0,0 +1,217 @@
+// Copyright Epic Games, Inc. All Rights Reserved.
+
+#include "Online/TitleFileEOSGS.h"
+
+#include "Online/AuthEOSGS.h"
+#include "Online/OnlineIdEOSGS.h"
+#include "Online/OnlineServicesEOSGS.h"
+#include "Online/OnlineServicesEOSGSTypes.h"
+#include "EOSSharedTypes.h"
+
+#include "eos_titlestorage.h"
+#include "eos_titlestorage_types.h"
+
+
+namespace UE::Online {
+
+FTitleFileEOSGS::FTitleFileEOSGS(FOnlineServicesEOSGS& InServices)
+	: Super(InServices)
+{
+}
+
+void FTitleFileEOSGS::Initialize()
+{
+	Super::Initialize();
+
+	TitleStorageHandle = EOS_Platform_GetTitleStorageInterface(static_cast<FOnlineServicesEOSGS&>(GetServices()).GetEOSPlatformHandle());
+	check(TitleStorageHandle);
+}
+
+void FTitleFileEOSGS::LoadConfig()
+{
+	Super::LoadConfig(Config);
+}
+
+TOnlineAsyncOpHandle<FTitleFileEnumerateFiles> FTitleFileEOSGS::EnumerateFiles(FTitleFileEnumerateFiles::Params&& InParams)
+{
+	TOnlineAsyncOpRef<FTitleFileEnumerateFiles> Op = GetOp<FTitleFileEnumerateFiles>(MoveTemp(InParams));
+
+	Op->Then([this](TOnlineAsyncOp<FTitleFileEnumerateFiles>& InAsyncOp)
+	{
+		const FTitleFileEnumerateFiles::Params& Params = InAsyncOp.GetParams();
+
+		if (!Services.Get<FAuthEOSGS>()->IsLoggedIn(Params.LocalUserId))
+		{
+			InAsyncOp.SetError(Errors::InvalidUser());
+			return MakeFulfilledPromise<const EOS_TitleStorage_QueryFileListCallbackInfo*>().GetFuture();
+		}
+
+		if (Config.SearchTag.IsEmpty())
+		{
+			InAsyncOp.SetError(Errors::NotConfigured());
+			return MakeFulfilledPromise<const EOS_TitleStorage_QueryFileListCallbackInfo *>().GetFuture();
+		}
+
+		const FTCHARToUTF8 Utf8SearchTag(*Config.SearchTag);
+		const char* SearchTagPtr = Utf8SearchTag.Get();
+
+		EOS_TitleStorage_QueryFileListOptions Options = {};
+		Options.ApiVersion = EOS_TITLESTORAGE_QUERYFILELISTOPTIONS_API_LATEST;
+		static_assert(EOS_TITLESTORAGE_QUERYFILELISTOPTIONS_API_LATEST == 1, "EOS_TitleStorage_QueryFileListOptions updated, check new fields");
+		Options.LocalUserId = GetProductUserIdChecked(InAsyncOp.GetParams().LocalUserId);
+		Options.ListOfTagsCount = 1;
+		Options.ListOfTags = &SearchTagPtr;
+
+		return EOS_Async<EOS_TitleStorage_QueryFileListCallbackInfo>(EOS_TitleStorage_QueryFileList, TitleStorageHandle, Options);
+	})
+	.Then([this](TOnlineAsyncOp<FTitleFileEnumerateFiles>& InAsyncOp, const EOS_TitleStorage_QueryFileListCallbackInfo* Data)
+	{
+		if (Data->ResultCode != EOS_EResult::EOS_Success)
+		{
+			UE_LOG(LogTemp, Warning, TEXT("EOS_TitleStorage_QueryFileList failed with result=[%s]"), *LexToString(Data->ResultCode));
+			InAsyncOp.SetError(FromEOSError(Data->ResultCode));
+			return;
+		}
+
+		EOS_ProductUserId LocalUserPuid = GetProductUserIdChecked(InAsyncOp.GetParams().LocalUserId);
+
+		EOS_TitleStorage_GetFileMetadataCountOptions GetCountOptions;
+		GetCountOptions.ApiVersion = EOS_TITLESTORAGE_GETFILEMETADATACOUNTOPTIONS_API_LATEST;
+		static_assert(EOS_TITLESTORAGE_GETFILEMETADATACOUNTOPTIONS_API_LATEST == 1, "EOS_TitleStorage_GetFileMetadataCountOptions updated, check new fields");
+		GetCountOptions.LocalUserId = LocalUserPuid;
+
+		const uint32_t NumFiles = EOS_TitleStorage_GetFileMetadataCount(TitleStorageHandle, &GetCountOptions);
+
+		TArray<FString> NewEnumeratedFiles;
+		NewEnumeratedFiles.Reserve(NumFiles);
+
+		for (uint32_t FileIdx = 0; FileIdx < NumFiles; FileIdx++)
+		{
+			EOS_TitleStorage_CopyFileMetadataAtIndexOptions CopyOptions;
+			CopyOptions.ApiVersion = EOS_TITLESTORAGE_COPYFILEMETADATAATINDEXOPTIONS_API_LATEST;
+			static_assert(EOS_TITLESTORAGE_COPYFILEMETADATAATINDEXOPTIONS_API_LATEST == 1, "EOS_TitleStorage_CopyFileMetadataAtIndexOptions updated, check new fields");
+			CopyOptions.LocalUserId = LocalUserPuid;
+			CopyOptions.Index = FileIdx;
+
+			EOS_TitleStorage_FileMetadata* FileMetadata = nullptr;
+
+			const EOS_EResult Result = EOS_TitleStorage_CopyFileMetadataAtIndex(TitleStorageHandle, &CopyOptions, &FileMetadata);
+			if (Result != EOS_EResult::EOS_Success)
+			{
+				UE_LOG(LogTemp, Warning, TEXT("EOS_TitleStorage_CopyFileMetadataAtIndex failed with result=[%s]"), *LexToString(Result));
+				InAsyncOp.SetError(FromEOSError(Result));
+				return;
+			}
+
+			NewEnumeratedFiles.Emplace(UTF8_TO_TCHAR(FileMetadata->Filename));
+
+			EOS_TitleStorage_FileMetadata_Release(FileMetadata);
+		}
+
+		EnumeratedFiles.Emplace(MoveTemp(NewEnumeratedFiles));
+		InAsyncOp.SetResult({});
+	})
+	.Enqueue(GetSerialQueue());
+	return Op->GetHandle();
+}
+
+TOnlineResult<FTitleFileGetEnumeratedFiles> FTitleFileEOSGS::GetEnumeratedFiles(FTitleFileGetEnumeratedFiles::Params&& Params)
+{
+	if (!Services.Get<FAuthEOSGS>()->IsLoggedIn(Params.LocalUserId))
+	{
+		return TOnlineResult<FTitleFileGetEnumeratedFiles>(Errors::InvalidUser());
+	}
+
+	if (!EnumeratedFiles.IsSet())
+	{
+		// Call EnumerateFiles first
+		return TOnlineResult<FTitleFileGetEnumeratedFiles>(Errors::InvalidState());
+	}
+
+	return TOnlineResult<FTitleFileGetEnumeratedFiles>({*EnumeratedFiles});
+}
+
+TOnlineAsyncOpHandle<FTitleFileReadFile> FTitleFileEOSGS::ReadFile(FTitleFileReadFile::Params&& InParams)
+{
+	static const TCHAR* FileContentsKey = TEXT("FileContents");
+
+	TOnlineAsyncOpRef<FTitleFileReadFile> Op = GetOp<FTitleFileReadFile>(MoveTemp(InParams));
+	const FTitleFileReadFile::Params& Params = Op->GetParams();
+
+	if (Params.Filename.IsEmpty())
+	{
+		Op->SetError(Errors::InvalidParams());
+		return Op->GetHandle();
+	}
+
+	Op->Then([this](TOnlineAsyncOp<FTitleFileReadFile>& InAsyncOp)
+	{
+		const FTitleFileReadFile::Params& Params = InAsyncOp.GetParams();
+
+		if (!Services.Get<FAuthEOSGS>()->IsLoggedIn(Params.LocalUserId))
+		{
+			InAsyncOp.SetError(Errors::InvalidUser());
+			return MakeFulfilledPromise<const EOS_TitleStorage_ReadFileCallbackInfo*>().GetFuture();
+		}
+
+		FTitleFileContentsRef FileContentsRef = MakeShared<FTitleFileContents>();
+		InAsyncOp.Data.Set<FTitleFileContentsRef>(FileContentsKey, FileContentsRef);
+
+		using FReadFileDataCallback = TEOSGlobalCallback<EOS_TitleStorage_OnReadFileDataCallback, EOS_TitleStorage_ReadFileDataCallbackInfo, TOnlineAsyncOp<FTitleFileReadFile>, EOS_TitleStorage_EReadResult>;
+		TSharedRef<FReadFileDataCallback> ReadFileDataCallback = MakeShared<FReadFileDataCallback>(InAsyncOp.AsWeak());
+		InAsyncOp.Data.Set<TSharedRef<FReadFileDataCallback>>(TEXT("ReadFileDataCallback"), ReadFileDataCallback);
+		ReadFileDataCallback->CallbackLambda = [FileContentsRef](const EOS_TitleStorage_ReadFileDataCallbackInfo* Data) -> EOS_TitleStorage_EReadResult
+		{
+			FTitleFileContents& FileContents = const_cast<FTitleFileContents&>(FileContentsRef.Get());
+			// Only has any effect on the first callback
+			FileContents.Reserve(Data->TotalFileSizeBytes);
+			FileContents.Append((uint8*)Data->DataChunk, Data->DataChunkLengthBytes);
+
+			UE_LOG(LogTemp, VeryVerbose, TEXT("FTitleFileEOSGS::ReadFile ReadProgress Filename=[%s] %d/%d"), UTF8_TO_TCHAR(Data->Filename), FileContents.Num(), Data->TotalFileSizeBytes);
+
+			return EOS_TitleStorage_EReadResult::EOS_TS_RR_ContinueReading;
+		};
+
+		const FTCHARToUTF8 Utf8Filename(Params.Filename);
+
+		EOS_TitleStorage_ReadFileOptions  Options = {};
+		Options.ApiVersion = EOS_TITLESTORAGE_READFILEOPTIONS_API_LATEST;
+		static_assert(EOS_TITLESTORAGE_READFILEOPTIONS_API_LATEST == 1, "EOS_TitleStorage_ReadFileOptions updated, check new fields");
+		Options.LocalUserId = GetProductUserIdChecked(InAsyncOp.GetParams().LocalUserId);
+		Options.Filename = Utf8Filename.Get();
+		Options.ReadChunkLengthBytes = 1024 * 1024;
+		Options.ReadFileDataCallback = ReadFileDataCallback->GetCallbackPtr();
+		Options.FileTransferProgressCallback = &FTitleFileEOSGS::OnFileTransferProgressStatic;
+
+		return EOS_Async<EOS_TitleStorage_ReadFileCallbackInfo>(EOS_TitleStorage_ReadFile, TitleStorageHandle, Options);
+	})
+	.Then([this](TOnlineAsyncOp<FTitleFileReadFile>& InAsyncOp, const EOS_TitleStorage_ReadFileCallbackInfo* Data)
+	{
+		if (Data->ResultCode != EOS_EResult::EOS_Success)
+		{
+			UE_LOG(LogTemp, Warning, TEXT("EOS_TitleStorage_ReadFile failed with result=[%s]"), *LexToString(Data->ResultCode));
+			InAsyncOp.SetError(FromEOSError(Data->ResultCode));
+			return;
+		}
+
+		const FTitleFileContentsRef* FileContents = InAsyncOp.Data.Get<FTitleFileContentsRef>(FileContentsKey);
+		if (ensure(FileContents))
+		{
+			InAsyncOp.SetResult({ *FileContents });
+		}
+		else
+		{
+			InAsyncOp.SetError(Errors::Unknown());
+			return;
+		}
+	})
+	.Enqueue(GetSerialQueue());
+	return Op->GetHandle();
+}
+
+void EOS_CALL FTitleFileEOSGS::OnFileTransferProgressStatic(const EOS_TitleStorage_FileTransferProgressCallbackInfo* Data)
+{
+	UE_LOG(LogTemp, VeryVerbose, TEXT("FTitleFileEOSGS::ReadFile TransferProgress Filename=[%s] %d/%d"), UTF8_TO_TCHAR(Data->Filename), Data->BytesTransferred, Data->TotalFileSizeBytes);
+}
+
+/* UE::Online */ }
diff --git a/Engine/Plugins/Online/Experimental/OnlineServicesEOSGS/Source/Public/Online/TitleFileEOSGS.h b/Engine/Plugins/Online/Experimental/OnlineServicesEOSGS/Source/Public/Online/TitleFileEOSGS.h
new file mode 100644
--- /dev/null
+++ b/Engine/Plugins/Online/Experimental/OnlineServicesEOSGS/Source/Public/Online/TitleFileEOSGS.h
@@ -0,0 +1,57 @@
+// Copyright Epic Games, Inc. All Rights Reserved.
+
+#pragma once
+
+#include "CoreMinimal.h"
+#include "Online/TitleFileCommon.h"
+
+#if defined(EOS_PLATFORM_BASE_FILE_NAME)
+#include EOS_PLATFORM_BASE_FILE_NAME
+#endif
+#include "eos_titlestorage_types.h"
+
+namespace UE::Online {
+
+class FOnlineServicesEOSGS;
+
+struct FTitleFileEOSGSConfig
+{
+	FString SearchTag;
+};
+
+namespace Meta {
+
+BEGIN_ONLINE_STRUCT_META(FTitleFileEOSGSConfig)
+	ONLINE_STRUCT_FIELD(FTitleFileEOSGSConfig, SearchTag)
+END_ONLINE_STRUCT_META()
+
+/* Meta*/ }
+
+class ONLINESERVICESEOSGS_API FTitleFileEOSGS : public FTitleFileCommon
+{
+public:
+	using Super = FTitleFileCommon;
+
+	FTitleFileEOSGS(FOnlineServicesEOSGS& InOwningSubsystem);
+	virtual ~FTitleFileEOSGS() = default;
+
+	// IOnlineComponent
+	virtual void Initialize() override;
+	virtual void LoadConfig() override;
+
+	// ITitleFile
+	virtual TOnlineAsyncOpHandle<FTitleFileEnumerateFiles> EnumerateFiles(FTitleFileEnumerateFiles::Params&& Params) override;
+	virtual TOnlineResult<FTitleFileGetEnumeratedFiles> GetEnumeratedFiles(FTitleFileGetEnumeratedFiles::Params&& Params) override;
+	virtual TOnlineAsyncOpHandle<FTitleFileReadFile> ReadFile(FTitleFileReadFile::Params&& Params) override;
+
+protected:
+	EOS_HTitleStorage TitleStorageHandle = nullptr;
+
+	TOptional<TArray<FString>> EnumeratedFiles;
+
+	static void EOS_CALL OnFileTransferProgressStatic(const EOS_TitleStorage_FileTransferProgressCallbackInfo* Data);
+
+	FTitleFileEOSGSConfig Config;
+};
+
+/* UE::Online */ }
