CIS fix

diff --git a/Engine/Plugins/Online/Experimental/OnlineServicesEOSGS/Source/Private/Online/OnlineServicesEOSGSPlatformFactory.h b/Engine/Plugins/Online/Experimental/OnlineServicesEOSGS/Source/Private/Online/OnlineServicesEOSGSPlatformFactory.h
--- a/Engine/Plugins/Online/Experimental/OnlineServicesEOSGS/Source/Private/Online/OnlineServicesEOSGSPlatformFactory.h
+++ b/Engine/Plugins/Online/Experimental/OnlineServicesEOSGS/Source/Private/Online/OnlineServicesEOSGSPlatformFactory.h
@@ -5,6 +5,7 @@
 #include "CoreMinimal.h"
 
 using IEOSPlatformHandlePtr = TSharedPtr<class IEOSPlatformHandle>;
+class FLazySingleton;
 
 namespace UE::Online
 {
