OSSv2 Sessions Interface - EOSGS Implementation - Find and Join functions

diff --git a/Engine/Plugins/Online/Experimental/OnlineServices/Source/OnlineServicesInterface/Private/Online/Sessions.cpp b/Engine/Plugins/Online/Experimental/OnlineServices/Source/OnlineServicesInterface/Private/Online/Sessions.cpp
--- a/Engine/Plugins/Online/Experimental/OnlineServices/Source/OnlineServicesInterface/Private/Online/Sessions.cpp
+++ b/Engine/Plugins/Online/Experimental/OnlineServices/Source/OnlineServicesInterface/Private/Online/Sessions.cpp
@@ -98,6 +98,7 @@ const TCHAR* LexToString(ESessionState Value)
 	switch (Value)
 	{
 	case ESessionState::Creating:	return TEXT("Creating");
+	case ESessionState::Joining:	return TEXT("Joining");
 	case ESessionState::Valid:		return TEXT("Valid");
 	case ESessionState::Leaving:	return TEXT("Leaving");
 	case ESessionState::Destroying:	return TEXT("Destroying");
@@ -116,6 +117,10 @@ void LexFromString(ESessionState& Value, const TCHAR* InStr)
 	{
 		Value = ESessionState::Creating;
 	}
+	else if (FCString::Stricmp(InStr, TEXT("Joining")) == 0)
+	{
+		Value = ESessionState::Joining;
+	}
 	else if (FCString::Stricmp(InStr, TEXT("Valid")) == 0)
 	{
 		Value = ESessionState::Valid;
diff --git a/Engine/Plugins/Online/Experimental/OnlineServices/Source/OnlineServicesInterface/Public/Online/Sessions.h b/Engine/Plugins/Online/Experimental/OnlineServices/Source/OnlineServicesInterface/Public/Online/Sessions.h
--- a/Engine/Plugins/Online/Experimental/OnlineServices/Source/OnlineServicesInterface/Public/Online/Sessions.h
+++ b/Engine/Plugins/Online/Experimental/OnlineServices/Source/OnlineServicesInterface/Public/Online/Sessions.h
@@ -91,6 +91,7 @@ enum class ESessionState : uint8
 {
 	Invalid,
 	Creating,
+	Joining,
 	Valid,
 	Leaving,
 	Destroying
@@ -213,7 +214,7 @@ struct ONLINESERVICESINTERFACE_API FSession
 	FOnlineSessionIdHandle SessionId;
 
 	/** The current state of the session */
-	ESessionState CurrentState;
+	ESessionState CurrentState = ESessionState::Invalid;
 
 	/** Set of session properties */
 	FSessionSettings SessionSettings;
diff --git a/Engine/Plugins/Online/Experimental/OnlineServicesEOSGS/Source/Private/Online/SessionsEOSGS.cpp b/Engine/Plugins/Online/Experimental/OnlineServicesEOSGS/Source/Private/Online/SessionsEOSGS.cpp
--- a/Engine/Plugins/Online/Experimental/OnlineServicesEOSGS/Source/Private/Online/SessionsEOSGS.cpp
+++ b/Engine/Plugins/Online/Experimental/OnlineServicesEOSGS/Source/Private/Online/SessionsEOSGS.cpp
@@ -113,6 +113,40 @@ FSessionEOSGS::FSessionEOSGS()
 FSessionEOSGS::FSessionEOSGS(const FSession& InSession)
 	: FSession(InSession)
 {
+	const FSessionEOSGS& Session = Cast(InSession);
+
+	SessionDetailsHandle = Session.SessionDetailsHandle;
+}
+
+FSessionEOSGS::FSessionEOSGS(const EOS_HSessionDetails& InSessionDetailsHandle)
+	: SessionDetailsHandle(MakeShareable<FSessionDetailsHandleEOSGS>(new FSessionDetailsHandleEOSGS(InSessionDetailsHandle)))
+{
+	EOS_SessionDetails_CopyInfoOptions CopyInfoOptions = { };
+	CopyInfoOptions.ApiVersion = EOS_SESSIONDETAILS_COPYINFO_API_LATEST;
+	static_assert(EOS_SESSIONDETAILS_COPYINFO_API_LATEST == 1, "EOS_SessionDetails_CopyInfoOptions updated, check new fields");
+
+	EOS_SessionDetails_Info* SessionDetailsInfo = nullptr;
+	EOS_EResult CopyInfoResult = EOS_SessionDetails_CopyInfo(InSessionDetailsHandle, &CopyInfoOptions, &SessionDetailsInfo);
+	if (CopyInfoResult == EOS_EResult::EOS_Success)
+	{
+		CurrentState = ESessionState::Valid;
+		SessionId = FOnlineSessionIdRegistryEOSGS::Get().FindOrAdd(FString(SessionDetailsInfo->SessionId));
+
+		// TODO: Load all session settings after SchemaVariant work
+	}
+	else
+	{
+		UE_LOG(LogTemp, Error, TEXT("[FSessionEOSGS] EOS_SessionDetails_CopyInfo failed with result [%s]"), *LexToString(CopyInfoResult));
+
+		CurrentState = ESessionState::Invalid;
+	}
+}
+
+const FSessionEOSGS& FSessionEOSGS::Cast(const FSession& InSession)
+{
+	check(InSession.SessionId.GetOnlineServicesType() == EOnlineServices::Epic);
+
+	return *static_cast<const FSessionEOSGS*>(&InSession);
 }
 
 /** FSessionsEOSGS */
@@ -171,7 +205,7 @@ TOnlineAsyncOpHandle<FCreateSession> FSessionsEOSGS::CreateSession(FCreateSessio
 {
 	TOnlineAsyncOpRef<FCreateSession> Op = GetOp<FCreateSession>(MoveTemp(Params));
 
-	Op->Then([this](TOnlineAsyncOp<FCreateSession>& Op) mutable
+	Op->Then([this](TOnlineAsyncOp<FCreateSession>& Op)
 	{
 		const FCreateSession::Params& OpParams = Op.GetParams();
 
@@ -289,11 +323,13 @@ TOnlineAsyncOpHandle<FCreateSession> FSessionsEOSGS::CreateSession(FCreateSessio
 		
 		return UpdateSessionImpl(MoveTemp(UpdateSessionImplParams));
 	})
-	.Then([this](TOnlineAsyncOp<FCreateSession>& Op, TDefaultErrorResult<FUpdateSessionImpl>&& Result) mutable
+	.Then([this](TOnlineAsyncOp<FCreateSession>& Op, TDefaultErrorResult<FUpdateSessionImpl>&& Result)
 	{
+		const FCreateSession::Params& OpParams = Op.GetParams();
+
 		if (Result.IsOk())
 		{
-			if (TSharedRef<FSessionEOSGS>* Session = SessionsByName.Find(Op.GetParams().SessionName))
+			if (TSharedRef<FSessionEOSGS>* Session = SessionsByName.Find(OpParams.SessionName))
 			{
 				Session->Get().CurrentState = ESessionState::Valid;
 
@@ -307,7 +343,14 @@ TOnlineAsyncOpHandle<FCreateSession> FSessionsEOSGS::CreateSession(FCreateSessio
 			}
 		}
 		else
-		{
+		{	
+			// If creation failed, we remove the session objects
+			if (TSharedRef<FSessionEOSGS>* Session = SessionsByName.Find(OpParams.SessionName))
+			{
+				SessionsByName.Remove(OpParams.SessionName);
+				SessionsById.Remove(Session->Get().SessionId);
+			}
+
 			Op.SetError(MoveTemp(Result.GetErrorValue()));
 		}
 	})
@@ -594,13 +637,13 @@ TOnlineAsyncOpHandle<FUpdateSession> FSessionsEOSGS::UpdateSession(FUpdateSessio
 	TOnlineAsyncOpRef<FUpdateSession> Op = GetOp<FUpdateSession>(MoveTemp(Params));
 	const FUpdateSession::Params& OpParams = Op->GetParams();
 
-	Op->Then([this](TOnlineAsyncOp<FUpdateSession>& Op) mutable
+	Op->Then([this](TOnlineAsyncOp<FUpdateSession>& Op)
 	{
 		const FUpdateSession::Params& OpParams = Op.GetParams();
 
 		if (TSharedRef<FSessionEOSGS>* FoundSession = SessionsByName.Find(OpParams.SessionName))
 		{
-			if (FoundSession->Get().CurrentState == ESessionState::Creating)
+			if (FoundSession->Get().CurrentState != ESessionState::Valid)
 			{
 				Op.SetError(Errors::InvalidState());
 
@@ -658,7 +701,7 @@ TOnlineAsyncOpHandle<FUpdateSession> FSessionsEOSGS::UpdateSession(FUpdateSessio
 			return MakeFulfilledPromise<TDefaultErrorResult<FUpdateSessionImpl>>().GetFuture();
 		}
 	})
-	.Then([this](TOnlineAsyncOp<FUpdateSession>& Op, TDefaultErrorResult<FUpdateSessionImpl>&& Result) mutable
+	.Then([this](TOnlineAsyncOp<FUpdateSession>& Op, TDefaultErrorResult<FUpdateSessionImpl>&& Result)
 	{
 		if (Result.IsOk())
 		{
@@ -727,7 +770,7 @@ TOnlineAsyncOpHandle<FLeaveSession> FSessionsEOSGS::LeaveSession(FLeaveSession::
 {
 	TOnlineAsyncOpRef<FLeaveSession> Op = GetOp<FLeaveSession>(MoveTemp(Params));
 
-	Op->Then([this](TOnlineAsyncOp<FLeaveSession>& Op) mutable
+	Op->Then([this](TOnlineAsyncOp<FLeaveSession>& Op)
 	{
 		const FLeaveSession::Params& OpParams = Op.GetParams();
 
@@ -751,7 +794,7 @@ TOnlineAsyncOpHandle<FLeaveSession> FSessionsEOSGS::LeaveSession(FLeaveSession::
 			return MakeFulfilledPromise<const EOS_Sessions_DestroySessionCallbackInfo*>().GetFuture();
 		}
 	})
-	.Then([this](TOnlineAsyncOp<FLeaveSession>& Op, const EOS_Sessions_DestroySessionCallbackInfo* Result) mutable
+	.Then([this](TOnlineAsyncOp<FLeaveSession>& Op, const EOS_Sessions_DestroySessionCallbackInfo* Result)
 	{
 		if (Result->ResultCode != EOS_EResult::EOS_Success)
 		{
@@ -783,4 +826,274 @@ TOnlineAsyncOpHandle<FLeaveSession> FSessionsEOSGS::LeaveSession(FLeaveSession::
 	return Op->GetHandle();
 }
 
+void FSessionsEOSGS::WriteSessionSearchHandle(FSessionSearchHandleEOSGS& SessionSearchHandle, const FFindSessions::Params& Params)
+{
+	if (Params.MaxResults > 0)
+	{
+		EOS_SessionSearch_SetMaxResultsOptions Options = { };
+		Options.ApiVersion = EOS_SESSIONSEARCH_SETMAXSEARCHRESULTS_API_LATEST;
+		static_assert(EOS_SESSIONSEARCH_SETMAXSEARCHRESULTS_API_LATEST == 1, "EOS_SessionSearch_SetMaxResultsOptions updated, check new fields");
+
+		// We truncate the max results parameter up to the value of EOS_SESSIONS_MAX_SEARCH_RESULTS
+		Options.MaxSearchResults = FMath::Clamp(Params.MaxResults, 1, EOS_SESSIONS_MAX_SEARCH_RESULTS);
+
+		EOS_EResult ResultCode = EOS_SessionSearch_SetMaxResults(SessionSearchHandle.SearchHandle, &Options);
+		if (ResultCode != EOS_EResult::EOS_Success)
+		{
+			UE_LOG(LogTemp, Warning, TEXT("[FSessionsEOSGS::WriteSessionSearchHandle] EOS_SessionSearch_SetMaxResults failed with result [%s]"), *LexToString(ResultCode));
+		}
+	}
+
+	if (!Params.Filters.IsEmpty())
+	{
+		// TODO: Pending SchemaVariant work
+		// Here we'll call EOS_SessionSearch_SetParameter
+	}
+
+	if (Params.SessionId.IsSet())
+	{
+		EOS_SessionSearch_SetSessionIdOptions Options = { };
+		Options.ApiVersion = EOS_SESSIONSEARCH_SETSESSIONID_API_LATEST;
+		static_assert(EOS_SESSIONSEARCH_SETSESSIONID_API_LATEST == 1, "EOS_SessionSearch_SetSessionIdOptions updated, check new fields");
+
+		FTCHARToUTF8 SessionIdUtf8(FOnlineSessionIdRegistryEOSGS::Get().ToLogString(*Params.SessionId));
+		Options.SessionId = SessionIdUtf8.Get();
+
+		EOS_EResult ResultCode = EOS_SessionSearch_SetSessionId(SessionSearchHandle.SearchHandle, &Options);
+		if (ResultCode != EOS_EResult::EOS_Success)
+		{
+			UE_LOG(LogTemp, Warning, TEXT("[FSessionsEOSGS::WriteSessionSearchHandle] EOS_SessionSearch_SetSessionId failed with result [%s]"), *LexToString(ResultCode));
+		}
+	}
+
+	if (Params.TargetUser.IsSet())
+	{
+		EOS_SessionSearch_SetTargetUserIdOptions Options = { };
+		Options.ApiVersion = EOS_SESSIONSEARCH_SETTARGETUSERID_API_LATEST;
+		static_assert(EOS_SESSIONSEARCH_SETTARGETUSERID_API_LATEST == 1, "EOS_SessionSearch_SetTargetUserIdOptions updated, check new fields");
+		Options.TargetUserId = GetProductUserIdChecked(*Params.TargetUser);
+
+		EOS_EResult ResultCode = EOS_SessionSearch_SetTargetUserId(SessionSearchHandle.SearchHandle, &Options);
+		if (ResultCode != EOS_EResult::EOS_Success)
+		{
+			UE_LOG(LogTemp, Warning, TEXT("[FSessionsEOSGS::WriteSessionSearchHandle] EOS_SessionSearch_SetTargetUserId failed with result [%s]"), *LexToString(ResultCode));
+		}
+	}
+}
+
+TOnlineAsyncOpHandle<FFindSessions> FSessionsEOSGS::FindSessions(FFindSessions::Params&& Params)
+{
+	TOnlineAsyncOpRef<FFindSessions> Op = GetOp<FFindSessions>(MoveTemp(Params));
+	const FFindSessions::Params& OpParams = Op->GetParams();
+
+	// We check if the user is logged in
+	IAuthPtr Auth = Services.GetAuthInterface();
+	TOnlineResult<FAuthGetAccountByAccountId> Result = Auth->GetAccountByAccountId({ OpParams.LocalUserId });
+
+	if (Result.IsOk())
+	{
+		if (Result.GetOkValue().AccountInfo->LoginStatus == ELoginStatus::NotLoggedIn)
+		{
+			UE_LOG(LogTemp, Warning, TEXT("[FSessionsEOSGS::FindSessions] Could not find sessions with user [%s] not logged in"), *ToLogString(OpParams.LocalUserId));
+
+			Op->SetError(Errors::InvalidUser());
+			return Op->GetHandle();
+		}
+	}
+
+	// We check if MaxResults has a valid value
+	if (OpParams.MaxResults > 0 && OpParams.MaxResults < EOS_SESSIONS_MAX_SEARCH_RESULTS)
+	{
+		UE_LOG(LogTemp, Error, TEXT("[FSessionsEOSGS::FindSessions] FFindSessions::Params::MaxResults had invalid value [%d]. Value must be between [0-%d]"), Params.MaxResults, EOS_SESSIONS_MAX_SEARCH_RESULTS);
+
+		Op->SetError(Errors::InvalidParams());
+		return Op->GetHandle();
+	}
+
+	// We start the find operation
+	Op->Then([this](TOnlineAsyncOp<FFindSessions>& Op)
+	{
+		if (!CurrentSessionSearch.IsValid())
+		{
+			CurrentSessionSearch = MakeShared<FFindSessions::Result>();
+			CurrentSessionSearchAsyncOpHandle = Op.AsShared();
+
+			const FFindSessions::Params& OpParams = Op.GetParams();
+
+			EOS_Sessions_CreateSessionSearchOptions CreateSessionSearchOptions = { };
+			CreateSessionSearchOptions.ApiVersion = EOS_SESSIONS_CREATESESSIONSEARCH_API_LATEST;
+			static_assert(EOS_SESSIONS_CREATESESSIONSEARCH_API_LATEST == 1, "EOS_Sessions_CreateSessionSearchOptions updated, check new fields");
+			CreateSessionSearchOptions.MaxSearchResults = FMath::Clamp(OpParams.MaxResults, 0, EOS_SESSIONS_MAX_SEARCH_RESULTS);
+
+			EOS_HSessionSearch SearchHandle = nullptr;
+			EOS_EResult ResultCode = EOS_Sessions_CreateSessionSearch(SessionsHandle, &CreateSessionSearchOptions, &SearchHandle);
+			if (ResultCode != EOS_EResult::EOS_Success)
+			{
+				UE_LOG(LogTemp, Error, TEXT("[FSessionsEOSGS::FindSessions] EOS_Sessions_CreateSessionSearch failed with result [%s]"), *LexToString(ResultCode));
+
+				Op.SetError(FromEOSError(ResultCode));
+
+				return MakeFulfilledPromise<const EOS_SessionSearch_FindCallbackInfo*>().GetFuture();
+			}
+
+			CurrentSessionSearchHandleEOSGS = MakeShareable(new FSessionSearchHandleEOSGS(SearchHandle));
+
+			// We write the search attributes
+			WriteSessionSearchHandle(*CurrentSessionSearchHandleEOSGS, OpParams);
+
+			EOS_SessionSearch_FindOptions FindOptions = { };
+			FindOptions.ApiVersion = EOS_SESSIONSEARCH_FIND_API_LATEST;
+			static_assert(EOS_SESSIONSEARCH_FIND_API_LATEST == 2, "EOS_SessionSearch_FindOptions updated, check new fields");
+			FindOptions.LocalUserId = GetProductUserIdChecked(OpParams.LocalUserId);
+
+			return EOS_Async<EOS_SessionSearch_FindCallbackInfo>(EOS_SessionSearch_Find, CurrentSessionSearchHandleEOSGS->SearchHandle, FindOptions);
+		}
+		else
+		{
+			Op.SetError(Errors::AlreadyPending());
+
+			return MakeFulfilledPromise<const EOS_SessionSearch_FindCallbackInfo*>().GetFuture();
+		}
+	})
+	.Then([this](TOnlineAsyncOp<FFindSessions>& Op, const EOS_SessionSearch_FindCallbackInfo*&& FindCallbackInfoResult)
+	{
+		if (FindCallbackInfoResult->ResultCode != EOS_EResult::EOS_Success)
+		{
+			Op.SetError(FromEOSError(FindCallbackInfoResult->ResultCode));
+			return;
+		}
+
+		// For a successful session, we'll get the search results
+		EOS_SessionSearch_GetSearchResultCountOptions GetSearchResultCountOptions = { };
+		GetSearchResultCountOptions.ApiVersion = EOS_SESSIONSEARCH_GETSEARCHRESULTCOUNT_API_LATEST;
+		static_assert(EOS_SESSIONSEARCH_GETSEARCHRESULTCOUNT_API_LATEST == 1, "EOS_SessionSearch_GetSearchResultCountOptions updated, check new fields");
+		int32 NumSearchResults = EOS_SessionSearch_GetSearchResultCount(CurrentSessionSearchHandleEOSGS->SearchHandle, &GetSearchResultCountOptions);
+
+		EOS_SessionSearch_CopySearchResultByIndexOptions CopySearchResultByIndexOptions = { };
+		CopySearchResultByIndexOptions.ApiVersion = EOS_SESSIONSEARCH_COPYSEARCHRESULTBYINDEX_API_LATEST;
+		static_assert(EOS_SESSIONSEARCH_COPYSEARCHRESULTBYINDEX_API_LATEST == 1, "EOS_SessionSearch_CopySearchResultByIndexOptions updated, check new fields");
+
+		for (int32 Index = 0; Index < NumSearchResults; Index++)
+		{
+			EOS_HSessionDetails SessionDetailsHandle = nullptr;
+
+			CopySearchResultByIndexOptions.SessionIndex = Index;
+			EOS_EResult CopySearchResultByIndexResult = EOS_SessionSearch_CopySearchResultByIndex(CurrentSessionSearchHandleEOSGS->SearchHandle, &CopySearchResultByIndexOptions, &SessionDetailsHandle);
+			if (CopySearchResultByIndexResult == EOS_EResult::EOS_Success)
+			{
+				// All session attributes get extracted with the SessionDetailsHandle
+				FSessionEOSGS NewSession = FSessionEOSGS(SessionDetailsHandle);
+
+				TSharedRef<FSessionEOSGS> NewSessionEOSGSRef = MakeShared<FSessionEOSGS>(NewSession);
+
+				CurrentSessionSearch->FoundSessions.Add(NewSessionEOSGSRef);
+			}
+			else
+			{
+				UE_LOG(LogTemp, Error, TEXT("[FSessionsEOSGS::FindSessions] EOS_SessionSearch_CopySearchResultByIndex failed for index [%d] with result [%s]"), Index, *LexToString(CopySearchResultByIndexResult));
+			}
+		}
+
+		Op.SetResult(MoveTemp(*CurrentSessionSearch));
+
+		CurrentSessionSearch.Reset();
+		CurrentSessionSearchHandleEOSGS.Reset();
+	})
+	.Enqueue(GetSerialQueue());
+
+	return Op->GetHandle();
+}
+
+TOnlineAsyncOpHandle<FJoinSession> FSessionsEOSGS::JoinSession(FJoinSession::Params&& Params)
+{
+	TOnlineAsyncOpRef<FJoinSession> Op = GetOp<FJoinSession>(MoveTemp(Params));
+	const FJoinSession::Params& OpParams = Op->GetParams();
+
+	// TODO: Call CheckJoinSessionParams after it's submitted, add the state check to it
+	// 
+	// We check that the passed session is valid
+	if (Params.Session->CurrentState != ESessionState::Valid)
+	{
+		UE_LOG(LogTemp, Warning, TEXT("[FSessionsEOSGS::FJoinSession] Could not join session with invalid session state [%s]"), *LexToString(Params.Session->CurrentState));
+
+		Op->SetError(Errors::InvalidParams());
+		return Op->GetHandle();
+	}
+
+	// We check that the passed session has a valid details handle
+	const FSessionEOSGS& SessionEOSGS = FSessionEOSGS::Cast(*Params.Session);
+	if (!SessionEOSGS.SessionDetailsHandle.IsValid())
+	{
+		UE_LOG(LogTemp, Warning, TEXT("[FSessionsEOSGS::FJoinSession] Could not join session with invalid session details handle in session with id [%s]"), *ToLogString(Params.Session->SessionId));
+
+		Op->SetError(Errors::InvalidParams());
+		return Op->GetHandle();
+	}
+
+	// TODO: Routine JoinSession checks
+
+	// We start the join operation
+	Op->Then([this](TOnlineAsyncOp<FJoinSession>& Op)
+	{
+		const FJoinSession::Params& OpParams = Op.GetParams();
+
+		// TODO: Call CheckJoinSessionState after it's submitted
+
+		// We make a copy of the Session object
+		FSessionEOSGS NewSession = FSessionEOSGS(*OpParams.Session);
+		NewSession.CurrentState = ESessionState::Joining;
+
+		TSharedRef<FSessionEOSGS> NewSessionEOSGSRef = MakeShared<FSessionEOSGS>(NewSession);
+
+		SessionsByName.Add(OpParams.SessionName, NewSessionEOSGSRef);
+		SessionsById.Add(NewSession.SessionId, NewSessionEOSGSRef);
+
+		// We start setup for the API call
+		EOS_Sessions_JoinSessionOptions JoinSessionOptions = { };
+		JoinSessionOptions.ApiVersion = EOS_SESSIONS_JOINSESSION_API_LATEST;
+		static_assert(EOS_SESSIONS_JOINSESSION_API_LATEST == 2, "EOS_Sessions_JoinSessionOptions updated, check new fields");
+
+		JoinSessionOptions.bPresenceEnabled = OpParams.Session->SessionSettings.bPresenceEnabled;
+
+		JoinSessionOptions.LocalUserId = GetProductUserIdChecked(OpParams.LocalUserId);
+
+		FTCHARToUTF8 SessionNameUtf8(OpParams.SessionName.ToString());
+		JoinSessionOptions.SessionName = SessionNameUtf8.Get();
+
+		const FSessionEOSGS& SessionEOSGS = FSessionEOSGS::Cast(*OpParams.Session);
+		JoinSessionOptions.SessionHandle = SessionEOSGS.SessionDetailsHandle->SessionDetailsHandle;
+
+		return EOS_Async<EOS_Sessions_JoinSessionCallbackInfo>(EOS_Sessions_JoinSession, SessionsHandle, JoinSessionOptions);
+	})
+	.Then([this](TOnlineAsyncOp<FJoinSession>& Op, const EOS_Sessions_JoinSessionCallbackInfo*&& Result)
+	{
+		const FJoinSession::Params& OpParams = Op.GetParams();
+
+		if (Result->ResultCode != EOS_EResult::EOS_Success)
+		{
+			// If join failed, we remove the session objects
+			if (TSharedRef<FSessionEOSGS>* Session = SessionsByName.Find(OpParams.SessionName))
+			{
+				SessionsByName.Remove(OpParams.SessionName);
+				SessionsById.Remove(Session->Get().SessionId);
+			}
+
+			Op.SetError(FromEOSError(Result->ResultCode));
+			return;
+		}
+		
+		if (TSharedRef<FSessionEOSGS>* FoundSession = SessionsByName.Find(OpParams.SessionName))
+		{
+			FoundSession->Get().CurrentState = ESessionState::Valid;
+		}
+
+		// A successful join allows the client to server travel, after which RegisterPlayers will be called by the engine
+
+		// TODO: Support multiple local users joining the session
+	})
+	.Enqueue(GetSerialQueue());
+
+	return Op->GetHandle();
+}
+
 /* UE::Online */ }
\ No newline at end of file
diff --git a/Engine/Plugins/Online/Experimental/OnlineServicesEOSGS/Source/Public/Online/SessionsEOSGS.h b/Engine/Plugins/Online/Experimental/OnlineServicesEOSGS/Source/Public/Online/SessionsEOSGS.h
--- a/Engine/Plugins/Online/Experimental/OnlineServicesEOSGS/Source/Public/Online/SessionsEOSGS.h
+++ b/Engine/Plugins/Online/Experimental/OnlineServicesEOSGS/Source/Public/Online/SessionsEOSGS.h
@@ -50,15 +50,48 @@ private:
 
 static FName EOS_SESSIONS_BUCKET_ID = TEXT("EOS_SESSIONS_BUCKET_ID");
 
+struct FSessionSearchHandleEOSGS : FNoncopyable
+{
+	EOS_HSessionSearch SearchHandle;
+
+	FSessionSearchHandleEOSGS(EOS_HSessionSearch InSearchHandle)
+		: SearchHandle(InSearchHandle)
+	{
+	}
+
+	virtual ~FSessionSearchHandleEOSGS()
+	{
+		EOS_SessionSearch_Release(SearchHandle);
+	}
+};
+
+struct FSessionDetailsHandleEOSGS : FNoncopyable
+{
+	EOS_HSessionDetails SessionDetailsHandle;
+
+	FSessionDetailsHandleEOSGS(EOS_HSessionDetails InSessionDetailsHandle)
+		: SessionDetailsHandle(InSessionDetailsHandle)
+	{
+	}
+
+	virtual ~FSessionDetailsHandleEOSGS()
+	{
+		EOS_SessionDetails_Release(SessionDetailsHandle);
+	}
+};
+
 class FSessionEOSGS : public FSession
 {
 public:
 	FSessionEOSGS();
 	FSessionEOSGS(const FSession& InSession);
+	FSessionEOSGS(const EOS_HSessionDetails& SessionDetailsHandle);
+
+	static const FSessionEOSGS& Cast(const FSession& InSession);
 
 public:
-	/** The IP address for the session owner */
-	TSharedPtr<FInternetAddr> OwnerInternetAddr;
+	/** Session details handle */
+	TSharedPtr<FSessionDetailsHandleEOSGS> SessionDetailsHandle;
 };
 
 struct FUpdateSessionImpl
@@ -94,6 +127,8 @@ public:
 	virtual TOnlineAsyncOpHandle<FCreateSession> CreateSession(FCreateSession::Params&& Params) override;
 	virtual TOnlineAsyncOpHandle<FUpdateSession> UpdateSession(FUpdateSession::Params&& Params) override;
 	virtual TOnlineAsyncOpHandle<FLeaveSession> LeaveSession(FLeaveSession::Params&& Params) override;
+	virtual TOnlineAsyncOpHandle<FFindSessions> FindSessions(FFindSessions::Params&& Params) override;
+	virtual TOnlineAsyncOpHandle<FJoinSession> JoinSession(FJoinSession::Params&& Params) override;
 
 protected:
 	void SetPermissionLevel(EOS_HSessionModification& SessionModHandle, const ESessionJoinPolicy& JoinPolicy);
@@ -110,6 +145,8 @@ protected:
 	void WriteUpdateSessionModificationHandle(EOS_HSessionModification& SessionModificationHandle, FSessionSettings& SessionSettings, const FSessionSettingsUpdate& NewSettings);
 	TFuture<TDefaultErrorResult<FUpdateSessionImpl>> UpdateSessionImpl(FUpdateSessionImpl::Params&& Params);
 
+	void WriteSessionSearchHandle(FSessionSearchHandleEOSGS& SessionSearchHandle, const FFindSessions::Params& Params);
+
 protected:
 	FOnlineServicesEOSGS& Services;
 
@@ -117,6 +154,10 @@ protected:
 
 	TMap<FName, TSharedRef<FSessionEOSGS>> SessionsByName;
 	TMap<FOnlineSessionIdHandle, TSharedRef<FSessionEOSGS>> SessionsById;
+
+	TSharedPtr<FSessionSearchHandleEOSGS> CurrentSessionSearchHandleEOSGS;
+	TSharedPtr<FFindSessions::Result> CurrentSessionSearch;
+	TSharedPtr<TOnlineAsyncOp<FFindSessions>> CurrentSessionSearchAsyncOpHandle;
 };
 
 /* UE::Online */ }
\ No newline at end of file
