- Fix TitleFileEOS not tracking and releasing the FileTransferRequest handle
- Make ReadChunkLengthBytes configurable.

diff --git a/Engine/Plugins/Online/Experimental/OnlineServicesEOSGS/Source/Private/Online/TitleFileEOSGS.cpp b/Engine/Plugins/Online/Experimental/OnlineServicesEOSGS/Source/Private/Online/TitleFileEOSGS.cpp
--- a/Engine/Plugins/Online/Experimental/OnlineServicesEOSGS/Source/Private/Online/TitleFileEOSGS.cpp
+++ b/Engine/Plugins/Online/Experimental/OnlineServicesEOSGS/Source/Private/Online/TitleFileEOSGS.cpp
@@ -58,7 +58,7 @@ TOnlineAsyncOpHandle<FTitleFileEnumerateFiles> FTitleFileEOSGS::EnumerateFiles(F
 		EOS_TitleStorage_QueryFileListOptions Options = {};
 		Options.ApiVersion = EOS_TITLESTORAGE_QUERYFILELISTOPTIONS_API_LATEST;
 		static_assert(EOS_TITLESTORAGE_QUERYFILELISTOPTIONS_API_LATEST == 1, "EOS_TitleStorage_QueryFileListOptions updated, check new fields");
-		Options.LocalUserId = GetProductUserIdChecked(InAsyncOp.GetParams().LocalUserId);
+		Options.LocalUserId = GetProductUserIdChecked(Params.LocalUserId);
 		Options.ListOfTagsCount = 1;
 		Options.ListOfTags = &SearchTagPtr;
 
@@ -66,6 +66,8 @@ TOnlineAsyncOpHandle<FTitleFileEnumerateFiles> FTitleFileEOSGS::EnumerateFiles(F
 	})
 	.Then([this](TOnlineAsyncOp<FTitleFileEnumerateFiles>& InAsyncOp, const EOS_TitleStorage_QueryFileListCallbackInfo* Data)
 	{
+		const FTitleFileEnumerateFiles::Params& Params = InAsyncOp.GetParams();
+
 		if (Data->ResultCode != EOS_EResult::EOS_Success)
 		{
 			UE_LOG(LogTemp, Warning, TEXT("EOS_TitleStorage_QueryFileList failed with result=[%s]"), *LexToString(Data->ResultCode));
@@ -73,7 +75,7 @@ TOnlineAsyncOpHandle<FTitleFileEnumerateFiles> FTitleFileEOSGS::EnumerateFiles(F
 			return;
 		}
 
-		EOS_ProductUserId LocalUserPuid = GetProductUserIdChecked(InAsyncOp.GetParams().LocalUserId);
+		const EOS_ProductUserId LocalUserPuid = GetProductUserIdChecked(Params.LocalUserId);
 
 		EOS_TitleStorage_GetFileMetadataCountOptions GetCountOptions;
 		GetCountOptions.ApiVersion = EOS_TITLESTORAGE_GETFILEMETADATACOUNTOPTIONS_API_LATEST;
@@ -134,6 +136,7 @@ TOnlineResult<FTitleFileGetEnumeratedFiles> FTitleFileEOSGS::GetEnumeratedFiles(
 TOnlineAsyncOpHandle<FTitleFileReadFile> FTitleFileEOSGS::ReadFile(FTitleFileReadFile::Params&& InParams)
 {
 	static const TCHAR* FileContentsKey = TEXT("FileContents");
+	static const TCHAR* RequestHandleKey = TEXT("RequestHandle");
 
 	TOnlineAsyncOpRef<FTitleFileReadFile> Op = GetOp<FTitleFileReadFile>(MoveTemp(InParams));
 	const FTitleFileReadFile::Params& Params = Op->GetParams();
@@ -177,32 +180,41 @@ TOnlineAsyncOpHandle<FTitleFileReadFile> FTitleFileEOSGS::ReadFile(FTitleFileRea
 		EOS_TitleStorage_ReadFileOptions  Options = {};
 		Options.ApiVersion = EOS_TITLESTORAGE_READFILEOPTIONS_API_LATEST;
 		static_assert(EOS_TITLESTORAGE_READFILEOPTIONS_API_LATEST == 1, "EOS_TitleStorage_ReadFileOptions updated, check new fields");
-		Options.LocalUserId = GetProductUserIdChecked(InAsyncOp.GetParams().LocalUserId);
+		Options.LocalUserId = GetProductUserIdChecked(Params.LocalUserId);
 		Options.Filename = Utf8Filename.Get();
-		Options.ReadChunkLengthBytes = 1024 * 1024;
+		Options.ReadChunkLengthBytes = Config.ReadChunkLengthBytes;
 		Options.ReadFileDataCallback = ReadFileDataCallback->GetCallbackPtr();
 		Options.FileTransferProgressCallback = &FTitleFileEOSGS::OnFileTransferProgressStatic;
 
-		return EOS_Async<EOS_TitleStorage_ReadFileCallbackInfo>(EOS_TitleStorage_ReadFile, TitleStorageHandle, Options);
+		TPromise<const EOS_TitleStorage_ReadFileCallbackInfo*> Promise;
+		TFuture<const EOS_TitleStorage_ReadFileCallbackInfo*> Future = Promise.GetFuture();
+		const EOS_HTitleStorageFileTransferRequest RequestHandle = EOS_Async<EOS_TitleStorage_ReadFileCallbackInfo>(EOS_TitleStorage_ReadFile, TitleStorageHandle, Options, MoveTemp(Promise));
+		InAsyncOp.Data.Set<EOS_HTitleStorageFileTransferRequest>(RequestHandleKey, RequestHandle);
+		return Future;
 	})
 	.Then([this](TOnlineAsyncOp<FTitleFileReadFile>& InAsyncOp, const EOS_TitleStorage_ReadFileCallbackInfo* Data)
 	{
-		if (Data->ResultCode != EOS_EResult::EOS_Success)
+		if (Data->ResultCode == EOS_EResult::EOS_Success)
+		{
+			const FTitleFileContentsRef* FileContents = InAsyncOp.Data.Get<FTitleFileContentsRef>(FileContentsKey);
+			if (ensure(FileContents))
+			{
+				InAsyncOp.SetResult({ *FileContents });
+			}
+			else
+			{
+				InAsyncOp.SetError(Errors::Unknown());
+			}
+		}
+		else
 		{
 			UE_LOG(LogTemp, Warning, TEXT("EOS_TitleStorage_ReadFile failed with result=[%s]"), *LexToString(Data->ResultCode));
 			InAsyncOp.SetError(FromEOSError(Data->ResultCode));
-			return;
 		}
 
-		const FTitleFileContentsRef* FileContents = InAsyncOp.Data.Get<FTitleFileContentsRef>(FileContentsKey);
-		if (ensure(FileContents))
+		if (const EOS_HTitleStorageFileTransferRequest* RequestHandle = InAsyncOp.Data.Get<EOS_HTitleStorageFileTransferRequest>(RequestHandleKey))
 		{
-			InAsyncOp.SetResult({ *FileContents });
-		}
-		else
-		{
-			InAsyncOp.SetError(Errors::Unknown());
-			return;
+			EOS_TitleStorageFileTransferRequest_Release(*RequestHandle);
 		}
 	})
 	.Enqueue(GetSerialQueue());
diff --git a/Engine/Plugins/Online/Experimental/OnlineServicesEOSGS/Source/Public/Online/OnlineServicesEOSGSTypes.h b/Engine/Plugins/Online/Experimental/OnlineServicesEOSGS/Source/Public/Online/OnlineServicesEOSGSTypes.h
--- a/Engine/Plugins/Online/Experimental/OnlineServicesEOSGS/Source/Public/Online/OnlineServicesEOSGSTypes.h
+++ b/Engine/Plugins/Online/Experimental/OnlineServicesEOSGS/Source/Public/Online/OnlineServicesEOSGSTypes.h
@@ -63,10 +63,10 @@ TFuture<const TEOSResult*> EOS_Async(TEOSFn EOSFn, TEOSHandle EOSHandle, TEOSPar
 }
 
 template<typename TEOSResult, typename TEOSHandle, typename TEOSParameters, typename TEOSFn>
-void EOS_Async(TEOSFn EOSFn, TEOSHandle EOSHandle, TEOSParameters Parameters, TPromise<const TEOSResult*>&& Promise)
+decltype(auto) EOS_Async(TEOSFn EOSFn, TEOSHandle EOSHandle, TEOSParameters Parameters, TPromise<const TEOSResult*>&& Promise)
 {
 	TEOSCallback<TEOSResult>* Callback = new TEOSCallback<TEOSResult>(MoveTemp(Promise));
-	EOSFn(EOSHandle, &Parameters, Callback, *Callback);
+	return EOSFn(EOSHandle, &Parameters, Callback, *Callback);
 }
 
 class EOSEventRegistration
diff --git a/Engine/Plugins/Online/Experimental/OnlineServicesEOSGS/Source/Public/Online/TitleFileEOSGS.h b/Engine/Plugins/Online/Experimental/OnlineServicesEOSGS/Source/Public/Online/TitleFileEOSGS.h
--- a/Engine/Plugins/Online/Experimental/OnlineServicesEOSGS/Source/Public/Online/TitleFileEOSGS.h
+++ b/Engine/Plugins/Online/Experimental/OnlineServicesEOSGS/Source/Public/Online/TitleFileEOSGS.h
@@ -17,15 +17,17 @@ class FOnlineServicesEOSGS;
 struct FTitleFileEOSGSConfig
 {
 	FString SearchTag;
+	int32 ReadChunkLengthBytes = 4096;
 };
 
 namespace Meta {
 
 BEGIN_ONLINE_STRUCT_META(FTitleFileEOSGSConfig)
-	ONLINE_STRUCT_FIELD(FTitleFileEOSGSConfig, SearchTag)
+	ONLINE_STRUCT_FIELD(FTitleFileEOSGSConfig, SearchTag),
+	ONLINE_STRUCT_FIELD(FTitleFileEOSGSConfig, ReadChunkLengthBytes)
 END_ONLINE_STRUCT_META()
 
-/* Meta*/ }
+/* Meta */ }
 
 class ONLINESERVICESEOSGS_API FTitleFileEOSGS : public FTitleFileCommon
 {
@@ -47,11 +49,11 @@ public:
 protected:
 	EOS_HTitleStorage TitleStorageHandle = nullptr;
 
+	FTitleFileEOSGSConfig Config;
+
 	TOptional<TArray<FString>> EnumeratedFiles;
 
 	static void EOS_CALL OnFileTransferProgressStatic(const EOS_TitleStorage_FileTransferProgressCallbackInfo* Data);
-
-	FTitleFileEOSGSConfig Config;
 };
 
 /* UE::Online */ }
