Add invite handling to LobbiesEOSGS

diff --git a/Engine/Plugins/Online/Experimental/OnlineServices/Source/OnlineServicesCommon/Private/Online/LobbiesCommon.cpp b/Engine/Plugins/Online/Experimental/OnlineServices/Source/OnlineServicesCommon/Private/Online/LobbiesCommon.cpp
--- a/Engine/Plugins/Online/Experimental/OnlineServices/Source/OnlineServicesCommon/Private/Online/LobbiesCommon.cpp
+++ b/Engine/Plugins/Online/Experimental/OnlineServices/Source/OnlineServicesCommon/Private/Online/LobbiesCommon.cpp
@@ -222,6 +222,11 @@ TOnlineEvent<void(const FLobbyInvitationRemoved&)> FLobbiesCommon::OnLobbyInvita
 	return LobbyEvents.OnLobbyInvitationRemoved;
 }
 
+TOnlineEvent<void(const FUILobbyJoinRequested&)> FLobbiesCommon::OnUILobbyJoinRequested()
+{
+	return LobbyEvents.OnUILobbyJoinRequested;
+}
+
 #if LOBBIES_FUNCTIONAL_TEST_ENABLED
 TOnlineAsyncOpHandle<FFunctionalTestLobbies> FLobbiesCommon::FunctionalTest(FFunctionalTestLobbies::Params&& Params)
 {
diff --git a/Engine/Plugins/Online/Experimental/OnlineServices/Source/OnlineServicesCommon/Public/Online/LobbiesCommon.h b/Engine/Plugins/Online/Experimental/OnlineServices/Source/OnlineServicesCommon/Public/Online/LobbiesCommon.h
--- a/Engine/Plugins/Online/Experimental/OnlineServices/Source/OnlineServicesCommon/Public/Online/LobbiesCommon.h
+++ b/Engine/Plugins/Online/Experimental/OnlineServices/Source/OnlineServicesCommon/Public/Online/LobbiesCommon.h
@@ -49,6 +49,7 @@ public:
 	virtual TOnlineEvent<void(const FLobbyMemberAttributesChanged&)> OnLobbyMemberAttributesChanged() override;
 	virtual TOnlineEvent<void(const FLobbyInvitationAdded&)> OnLobbyInvitationAdded() override;
 	virtual TOnlineEvent<void(const FLobbyInvitationRemoved&)> OnLobbyInvitationRemoved() override;
+	virtual TOnlineEvent<void(const FUILobbyJoinRequested&)> OnUILobbyJoinRequested() override;
 
 protected:
 #if LOBBIES_FUNCTIONAL_TEST_ENABLED
diff --git a/Engine/Plugins/Online/Experimental/OnlineServices/Source/OnlineServicesCommon/Public/Online/LobbiesCommonTypes.h b/Engine/Plugins/Online/Experimental/OnlineServices/Source/OnlineServicesCommon/Public/Online/LobbiesCommonTypes.h
--- a/Engine/Plugins/Online/Experimental/OnlineServices/Source/OnlineServicesCommon/Public/Online/LobbiesCommonTypes.h
+++ b/Engine/Plugins/Online/Experimental/OnlineServices/Source/OnlineServicesCommon/Public/Online/LobbiesCommonTypes.h
@@ -106,6 +106,7 @@ struct FLobbyEvents final
 	TOnlineEventCallable<void(const FLobbyMemberAttributesChanged&)> OnLobbyMemberAttributesChanged;
 	TOnlineEventCallable<void(const FLobbyInvitationAdded&)> OnLobbyInvitationAdded;
 	TOnlineEventCallable<void(const FLobbyInvitationRemoved&)> OnLobbyInvitationRemoved;
+	TOnlineEventCallable<void(const FUILobbyJoinRequested&)> OnUILobbyJoinRequested;
 };
 
 class ONLINESERVICESCOMMON_API FLobbySchema final
diff --git a/Engine/Plugins/Online/Experimental/OnlineServices/Source/OnlineServicesInterface/Private/Online/Lobbies.cpp b/Engine/Plugins/Online/Experimental/OnlineServices/Source/OnlineServicesInterface/Private/Online/Lobbies.cpp
--- a/Engine/Plugins/Online/Experimental/OnlineServices/Source/OnlineServicesInterface/Private/Online/Lobbies.cpp
+++ b/Engine/Plugins/Online/Experimental/OnlineServices/Source/OnlineServicesInterface/Private/Online/Lobbies.cpp
@@ -209,6 +209,33 @@ void LexFromString(ELobbyComparisonOp& OutComparison, const TCHAR* InStr)
 	}
 }
 
+const TCHAR* LexToString(EUILobbyJoinRequestedSource UILobbyJoinRequestedSource)
+{
+	switch (UILobbyJoinRequestedSource)
+	{
+	case EUILobbyJoinRequestedSource::FromInvitation:	return TEXT("FromInvitation");
+	default:											checkNoEntry(); // Intentional fallthrough
+	case EUILobbyJoinRequestedSource::Unspecified:		return TEXT("Unspecified");
+	}
+}
+
+void LexFromString(EUILobbyJoinRequestedSource& OutUILobbyJoinRequestedSource, const TCHAR* InStr)
+{
+	if (FCString::Stricmp(InStr, TEXT("FromInvitation")) == 0)
+	{
+		OutUILobbyJoinRequestedSource = EUILobbyJoinRequestedSource::FromInvitation;
+	}
+	else if (FCString::Stricmp(InStr, TEXT("Unspecified")) == 0)
+	{
+		OutUILobbyJoinRequestedSource = EUILobbyJoinRequestedSource::Unspecified;
+	}
+	else
+	{
+		checkNoEntry();
+		OutUILobbyJoinRequestedSource = EUILobbyJoinRequestedSource::Unspecified;
+	}
+}
+
 void SortLobbies(const TArray<FFindLobbySearchFilter>& Filters, TArray<TSharedRef<const FLobby>>& Lobbies)
 {
 	// todo
diff --git a/Engine/Plugins/Online/Experimental/OnlineServices/Source/OnlineServicesInterface/Public/Online/Lobbies.h b/Engine/Plugins/Online/Experimental/OnlineServices/Source/OnlineServicesInterface/Public/Online/Lobbies.h
--- a/Engine/Plugins/Online/Experimental/OnlineServices/Source/OnlineServicesInterface/Public/Online/Lobbies.h
+++ b/Engine/Plugins/Online/Experimental/OnlineServices/Source/OnlineServicesInterface/Public/Online/Lobbies.h
@@ -596,6 +596,30 @@ struct FLobbyInvitationRemoved
 
 };
 
+/** Lobby join requested source */
+enum class EUILobbyJoinRequestedSource : uint8
+{
+	/** Unspecified by the online service */
+	Unspecified,
+	/** From an invitation */
+	FromInvitation,
+};
+ONLINESERVICESINTERFACE_API const TCHAR* LexToString(EUILobbyJoinRequestedSource UILobbyJoinRequestedSource);
+ONLINESERVICESINTERFACE_API void LexFromString(EUILobbyJoinRequestedSource& OutUILobbyJoinRequestedSource, const TCHAR* InStr);
+
+/** Struct for UILobbyJoinedRequested event */
+struct FUILobbyJoinRequested
+{
+	/** The local user associated with the join request. */
+	FOnlineAccountIdHandle LocalUserId;
+
+	/** The lobby the local user requested to join, or the online error if there was a failure */
+	TResult<TSharedRef<const FLobby>, FOnlineError> Result;
+
+	/** Join request source */
+	EUILobbyJoinRequestedSource JoinRequestedSource = EUILobbyJoinRequestedSource::Unspecified;
+};
+
 class ILobbies
 {
 public:
@@ -811,6 +835,13 @@ public:
 	 * @return Event that can be bound to
 	 */
 	virtual TOnlineEvent<void(const FLobbyInvitationRemoved&)> OnLobbyInvitationRemoved() = 0;
+
+	/**
+	 * Get the event that is triggered when a join is requested through an external mechanism.
+	 *
+	 * @return Event that can be bound to
+	 */
+	virtual TOnlineEvent<void(const FUILobbyJoinRequested&)> OnUILobbyJoinRequested() = 0;
 };
 
 namespace Meta {
@@ -926,11 +957,65 @@ BEGIN_ONLINE_STRUCT_META(FGetJoinedLobbies::Result)
 END_ONLINE_STRUCT_META()
 
 BEGIN_ONLINE_STRUCT_META(FGetReceivedInvitations::Params)
+	ONLINE_STRUCT_FIELD(FGetReceivedInvitations::Params, LocalUserId)
 END_ONLINE_STRUCT_META()
 
 BEGIN_ONLINE_STRUCT_META(FGetReceivedInvitations::Result)
 END_ONLINE_STRUCT_META()
 
+BEGIN_ONLINE_STRUCT_META(FLobbyJoined)
+	ONLINE_STRUCT_FIELD(FLobbyJoined, Lobby)
+END_ONLINE_STRUCT_META()
+
+BEGIN_ONLINE_STRUCT_META(FLobbyLeft)
+	ONLINE_STRUCT_FIELD(FLobbyLeft, Lobby)
+END_ONLINE_STRUCT_META()
+
+BEGIN_ONLINE_STRUCT_META(FLobbyMemberJoined)
+	ONLINE_STRUCT_FIELD(FLobbyMemberJoined, Lobby),
+	ONLINE_STRUCT_FIELD(FLobbyMemberJoined, Member)
+END_ONLINE_STRUCT_META()
+
+BEGIN_ONLINE_STRUCT_META(FLobbyMemberLeft)
+	ONLINE_STRUCT_FIELD(FLobbyMemberLeft, Lobby),
+	ONLINE_STRUCT_FIELD(FLobbyMemberLeft, Member),
+	ONLINE_STRUCT_FIELD(FLobbyMemberLeft, Reason)
+END_ONLINE_STRUCT_META()
+
+BEGIN_ONLINE_STRUCT_META(FLobbyLeaderChanged)
+	ONLINE_STRUCT_FIELD(FLobbyLeaderChanged, Lobby),
+	ONLINE_STRUCT_FIELD(FLobbyLeaderChanged, Leader)
+END_ONLINE_STRUCT_META()
+
+BEGIN_ONLINE_STRUCT_META(FLobbyAttributesChanged)
+	ONLINE_STRUCT_FIELD(FLobbyAttributesChanged, Lobby),
+	ONLINE_STRUCT_FIELD(FLobbyAttributesChanged, ChangedAttributes)
+END_ONLINE_STRUCT_META()
+
+BEGIN_ONLINE_STRUCT_META(FLobbyMemberAttributesChanged)
+	ONLINE_STRUCT_FIELD(FLobbyMemberAttributesChanged, Lobby),
+	ONLINE_STRUCT_FIELD(FLobbyMemberAttributesChanged, Member),
+	ONLINE_STRUCT_FIELD(FLobbyMemberAttributesChanged, ChangedAttributes)
+END_ONLINE_STRUCT_META()
+
+BEGIN_ONLINE_STRUCT_META(FLobbyInvitationAdded)
+	ONLINE_STRUCT_FIELD(FLobbyInvitationAdded, LocalUserId),
+	ONLINE_STRUCT_FIELD(FLobbyInvitationAdded, SenderId),
+	ONLINE_STRUCT_FIELD(FLobbyInvitationAdded, Lobby)
+END_ONLINE_STRUCT_META()
+
+BEGIN_ONLINE_STRUCT_META(FLobbyInvitationRemoved)
+	ONLINE_STRUCT_FIELD(FLobbyInvitationRemoved, LocalUserId),
+	ONLINE_STRUCT_FIELD(FLobbyInvitationRemoved, SenderId),
+	ONLINE_STRUCT_FIELD(FLobbyInvitationRemoved, Lobby)
+END_ONLINE_STRUCT_META()
+
+BEGIN_ONLINE_STRUCT_META(FUILobbyJoinRequested)
+	ONLINE_STRUCT_FIELD(FUILobbyJoinRequested, LocalUserId),
+	ONLINE_STRUCT_FIELD(FUILobbyJoinRequested, Result)
+END_ONLINE_STRUCT_META()
+
+
 /* Meta*/ }
 
 /* UE::Online */ }
diff --git a/Engine/Plugins/Online/Experimental/OnlineServicesEOSGS/Source/Private/Online/LobbiesEOSGS.cpp b/Engine/Plugins/Online/Experimental/OnlineServicesEOSGS/Source/Private/Online/LobbiesEOSGS.cpp
--- a/Engine/Plugins/Online/Experimental/OnlineServicesEOSGS/Source/Private/Online/LobbiesEOSGS.cpp
+++ b/Engine/Plugins/Online/Experimental/OnlineServicesEOSGS/Source/Private/Online/LobbiesEOSGS.cpp
@@ -978,42 +978,63 @@ void FLobbiesEOSGS::HandleLobbyInviteReceived(const EOS_Lobby_LobbyInviteReceive
 
 void FLobbiesEOSGS::HandleLobbyInviteAccepted(const EOS_Lobby_LobbyInviteAcceptedCallbackInfo* Data)
 {
-#if 0
-	EOS_STRUCT(EOS_Lobby_LobbyInviteAcceptedCallbackInfo, (
-		/** Context that was passed into EOS_Lobby_AddNotifyLobbyInviteAccepted */
-		void* ClientData;
-		/** The invite ID */
-		const char* InviteId;
-		/** The Product User ID of the local user who received the invitation */
-		EOS_ProductUserId LocalUserId;
-		/** The Product User ID of the user who sent the invitation */
-		EOS_ProductUserId TargetUserId;
-		/** Lobby ID that the user has been invited to */
-		const char* LobbyId;
-	));
-#endif
-
-	// Todo: handle catalog of sent invitations.
+	// Todo: Queue this like an operation.
+	const FOnlineAccountIdHandle LocalUserId = FindAccountId(Data->LocalUserId);
+	if (LocalUserId.IsValid())
+	{
+		FLobbyInviteDataEOS::CreateFromInviteId(LobbyPrerequisites.ToSharedRef(), LobbyDataRegistry.ToSharedRef(), LocalUserId, Data->InviteId, Data->TargetUserId)
+			.Then([this](TFuture<TDefaultErrorResultInternal<TSharedRef<FLobbyInviteDataEOS>>>&& Future)
+				{
+					if (Future.Get().IsError())
+					{
+						// Todo: Log / queue a manual fetch of invitations.
+						UE_LOG(LogTemp, Warning, TEXT("[FLobbiesEOSGS::HandleLobbyInviteReceived] Failed to receive invite. Error: %s"),
+							*Future.Get().GetErrorValue().GetLogString());
+					}
+					else
+					{
+						UE_LOG(LogTemp, Warning, TEXT("[FLobbiesEOSGS::HandleLobbyInviteReceived] Received invite. Id: %s, Lobby: %s, Receiver: %s, Sender: %s"),
+							*Future.Get().GetOkValue()->GetInviteId(),
+							*Future.Get().GetOkValue()->GetLobbyData()->GetLobbyId(),
+							*ToLogString(Future.Get().GetOkValue()->GetReceiver()),
+							*ToLogString(Future.Get().GetOkValue()->GetSender()));
+						const TSharedRef<FLobbyInviteDataEOS>& Invite = Future.Get().GetOkValue();
+
+						LobbyEvents.OnUILobbyJoinRequested.Broadcast(
+							FUILobbyJoinRequested{
+								Invite->GetReceiver(),
+								TResult<TSharedRef<const FLobby>, FOnlineError>(Invite->GetLobbyData()->GetClientLobbyData()->GetPublicDataPtr()),
+								EUILobbyJoinRequestedSource::FromInvitation
+							});
+					}
+				});
+	}
 }
 
 void FLobbiesEOSGS::HandleJoinLobbyAccepted(const EOS_Lobby_JoinLobbyAcceptedCallbackInfo* Data)
 {
-#if 0
-	EOS_STRUCT(EOS_Lobby_JoinLobbyAcceptedCallbackInfo, (
-		/** Context that was passed into EOS_Lobby_AddNotifyJoinLobbyAccepted */
-		void* ClientData;
-		/** The Product User ID of the local user who is joining */
-		EOS_ProductUserId LocalUserId;
-		/** 
-		 * The UI Event associated with this Join Game event.
-		 * This should be used with EOS_Lobby_CopyLobbyDetailsHandleByUiEventId to get a handle to be used
-		 * when calling EOS_Lobby_JoinLobby.
-		 */
-		EOS_UI_EventId UiEventId;
-	));
-#endif
-
-	// Todo: handle UI events.
+	// Todo: Queue this like an operation.
+	const FOnlineAccountIdHandle LocalUserId = FindAccountId(Data->LocalUserId);
+	if (LocalUserId.IsValid())
+	{
+		TDefaultErrorResultInternal<TSharedRef<FLobbyDetailsEOS>> LobbyDetailsResult = 
+			FLobbyDetailsEOS::CreateFromUiEventId(LobbyPrerequisites.ToSharedRef(), LocalUserId, Data->UiEventId);
+		if (LobbyDetailsResult.IsError())
+		{
+			// Todo: Log / queue a manual fetch of invitations.
+			UE_LOG(LogTemp, Warning, TEXT("[FLobbiesEOSGS::HandleJoinLobbyAccepted] Failed to receive invite. Error: %s"),
+				*LobbyDetailsResult.GetErrorValue().GetLogString());
+		}
+		LobbyDataRegistry->FindOrCreateFromLobbyDetails(LocalUserId, LobbyDetailsResult.GetOkValue()).Then([this, LocalUserId](TFuture<TDefaultErrorResultInternal<TSharedRef<FLobbyDataEOS>>> && Result)
+		{
+			LobbyEvents.OnUILobbyJoinRequested.Broadcast(
+				FUILobbyJoinRequested{
+					LocalUserId,
+					TResult<TSharedRef<const FLobby>, FOnlineError>(Result.Get().GetOkValue()->GetClientLobbyData()->GetPublicDataPtr()),
+					EUILobbyJoinRequestedSource::Unspecified
+				});
+		});
+	}
 }
 
 void FLobbiesEOSGS::RegisterHandlers()
@@ -1058,7 +1079,7 @@ void FLobbiesEOSGS::RegisterHandlers()
 	OnLobbyInviteAcceptedEOSEventRegistration = EOS_RegisterComponentEventHandler(
 		this,
 		LobbyPrerequisites->LobbyInterfaceHandle,
-		EOS_LOBBY_ADDNOTIFYLOBBYINVITERECEIVED_API_LATEST,
+		EOS_LOBBY_ADDNOTIFYLOBBYINVITEACCEPTED_API_LATEST,
 		&EOS_Lobby_AddNotifyLobbyInviteAccepted,
 		&EOS_Lobby_RemoveNotifyLobbyInviteAccepted,
 		&FLobbiesEOSGS::HandleLobbyInviteAccepted);
@@ -1272,7 +1293,7 @@ TOnlineAsyncOpHandle<FLobbiesEOSGS::FJoinLobbyMemberImpl> FLobbiesEOSGS::JoinLob
 		JoinLobbyOptions.ApiVersion = EOS_LOBBY_JOINLOBBY_API_LATEST;
 		JoinLobbyOptions.LobbyDetailsHandle = LobbyDetails->GetEOSHandle();
 		JoinLobbyOptions.LocalUserId = GetProductUserIdChecked(Params.LocalUserId);
-		JoinLobbyOptions.bPresenceEnabled = false;
+		JoinLobbyOptions.bPresenceEnabled = false; // todo
 		JoinLobbyOptions.LocalRTCOptions = nullptr;
 		EOS_Async(EOS_Lobby_JoinLobby, LobbyPrerequisites->LobbyInterfaceHandle, JoinLobbyOptions, MoveTemp(Promise));
 	})
