- UserFile EOS Implementation

diff --git a/Engine/Plugins/Online/EOSShared/Source/EOSShared/Public/EOSSharedTypes.h b/Engine/Plugins/Online/EOSShared/Source/EOSShared/Public/EOSSharedTypes.h
--- a/Engine/Plugins/Online/EOSShared/Source/EOSShared/Public/EOSSharedTypes.h
+++ b/Engine/Plugins/Online/EOSShared/Source/EOSShared/Public/EOSSharedTypes.h
@@ -27,12 +27,12 @@ public:
 #if WITH_EOS_SDK
 
 /** Class to handle all callbacks generically using a lambda to process callback results */
-template<typename CallbackFuncType, typename CallbackParamType, typename OwningType, typename CallbackReturnType = void>
+template<typename CallbackFuncType, typename CallbackParamType, typename OwningType, typename CallbackReturnType = void, typename... CallbackExtraParams>
 class TEOSGlobalCallback :
 	public FCallbackBase
 {
 public:
-	TFunction<CallbackReturnType(const CallbackParamType*)> CallbackLambda;
+	TFunction<CallbackReturnType(const CallbackParamType*, CallbackExtraParams... ExtraParams)> CallbackLambda;
 
 	TEOSGlobalCallback(TWeakPtr<OwningType> InOwner)
 		: FCallbackBase()
@@ -53,7 +53,7 @@ private:
 	/** The object that needs to be checked for lifetime before calling the callback */
 	TWeakPtr<OwningType> Owner;
 
-	static CallbackReturnType EOS_CALL CallbackImpl(const CallbackParamType* Data)
+	static CallbackReturnType EOS_CALL CallbackImpl(const CallbackParamType* Data, CallbackExtraParams... ExtraParams)
 	{
 		TEOSGlobalCallback* CallbackThis = (TEOSGlobalCallback*)Data->ClientData;
 		check(CallbackThis);
@@ -69,11 +69,11 @@ private:
 
 			if constexpr (std::is_void<CallbackReturnType>::value)
 			{
-				CallbackThis->CallbackLambda(Data);
+				CallbackThis->CallbackLambda(Data, ExtraParams...);
 			}
 			else
 			{
-				return CallbackThis->CallbackLambda(Data);
+				return CallbackThis->CallbackLambda(Data, ExtraParams...);
 			}
 		}
 		
diff --git a/Engine/Plugins/Online/Experimental/OnlineServicesEOSGS/Source/Private/Online/UserFileEOSGS.cpp b/Engine/Plugins/Online/Experimental/OnlineServicesEOSGS/Source/Private/Online/UserFileEOSGS.cpp
new file mode 100644
--- /dev/null
+++ b/Engine/Plugins/Online/Experimental/OnlineServicesEOSGS/Source/Private/Online/UserFileEOSGS.cpp
@@ -0,0 +1,424 @@
+// Copyright Epic Games, Inc. All Rights Reserved.
+
+#include "Online/UserFileEOSGS.h"
+
+#include "Online/AuthEOSGS.h"
+#include "Online/OnlineIdEOSGS.h"
+#include "Online/OnlineServicesEOSGS.h"
+#include "Online/OnlineServicesEOSGSTypes.h"
+#include "EOSSharedTypes.h"
+#include "HAL/UnrealMemory.h"
+
+#include "eos_playerdatastorage.h"
+#include "eos_playerdatastorage_types.h"
+
+
+namespace UE::Online {
+
+static const TCHAR* RequestHandleKey = TEXT("RequestHandle");
+
+FUserFileEOSGS::FUserFileEOSGS(FOnlineServicesEOSGS& InServices)
+	: Super(InServices)
+{
+}
+
+void FUserFileEOSGS::Initialize()
+{
+	Super::Initialize();
+
+	PlayerDataStorageHandle = EOS_Platform_GetPlayerDataStorageInterface(static_cast<FOnlineServicesEOSGS&>(GetServices()).GetEOSPlatformHandle());
+	check(PlayerDataStorageHandle);
+}
+
+void FUserFileEOSGS::LoadConfig()
+{
+	Super::LoadConfig(Config);
+}
+
+TOnlineAsyncOpHandle<FUserFileEnumerateFiles> FUserFileEOSGS::EnumerateFiles(FUserFileEnumerateFiles::Params&& InParams)
+{
+	TOnlineAsyncOpRef<FUserFileEnumerateFiles> Op = GetOp<FUserFileEnumerateFiles>(MoveTemp(InParams));
+
+	Op->Then([this](TOnlineAsyncOp<FUserFileEnumerateFiles>& Op, TPromise<const EOS_PlayerDataStorage_QueryFileListCallbackInfo*>&& Promise)
+	{
+		const FUserFileEnumerateFiles::Params& Params = Op.GetParams();
+
+		if (!Services.Get<FAuthEOSGS>()->IsLoggedIn(Params.LocalUserId))
+		{
+			Op.SetError(Errors::InvalidUser());
+			Promise.EmplaceValue();
+			return;
+		}
+
+		EOS_PlayerDataStorage_QueryFileListOptions Options = {};
+		Options.ApiVersion = EOS_PLAYERDATASTORAGE_QUERYFILELISTOPTIONS_API_LATEST;
+		static_assert(EOS_PLAYERDATASTORAGE_QUERYFILELISTOPTIONS_API_LATEST == 1, "EOS_PlayerDataStorage_QueryFileListOptions updated, check new fields");
+		Options.LocalUserId = GetProductUserIdChecked(Params.LocalUserId);
+
+		EOS_Async(EOS_PlayerDataStorage_QueryFileList, PlayerDataStorageHandle, Options, MoveTemp(Promise));
+	})
+	.Then([this](TOnlineAsyncOp<FUserFileEnumerateFiles>& Op, const EOS_PlayerDataStorage_QueryFileListCallbackInfo* Data)
+	{
+		const FUserFileEnumerateFiles::Params& Params = Op.GetParams();
+
+		if (Data->ResultCode != EOS_EResult::EOS_Success)
+		{
+			UE_LOG(LogTemp, Warning, TEXT("EOS_PlayerDataStorage_QueryFileList failed with result=[%s]"), *LexToString(Data->ResultCode));
+			Op.SetError(FromEOSError(Data->ResultCode));
+			return;
+		}
+
+		const EOS_ProductUserId LocalUserPuid = GetProductUserIdChecked(Params.LocalUserId);
+
+		EOS_PlayerDataStorage_GetFileMetadataCountOptions GetCountOptions;
+		GetCountOptions.ApiVersion = EOS_PLAYERDATASTORAGE_GETFILEMETADATACOUNTOPTIONS_API_LATEST;
+		static_assert(EOS_PLAYERDATASTORAGE_GETFILEMETADATACOUNTOPTIONS_API_LATEST == 1, "EOS_PlayerDataStorage_GetFileMetadataCountOptions updated, check new fields");
+		GetCountOptions.LocalUserId = LocalUserPuid;
+
+		int32_t NumFiles;
+		EOS_EResult Result = EOS_PlayerDataStorage_GetFileMetadataCount(PlayerDataStorageHandle, &GetCountOptions, &NumFiles);
+		if (Result != EOS_EResult::EOS_Success)
+		{
+			UE_LOG(LogTemp, Warning, TEXT("EOS_PlayerDataStorage_GetFileMetadataCount failed with result=[%s]"), *LexToString(Result));
+			Op.SetError(FromEOSError(Result));
+			return;
+		}
+
+		check(NumFiles >= 0);
+
+		TArray<FString> NewEnumeratedFiles;
+		NewEnumeratedFiles.Reserve(NumFiles);
+
+		for (int32_t FileIdx = 0; FileIdx < NumFiles; FileIdx++)
+		{
+			EOS_PlayerDataStorage_CopyFileMetadataAtIndexOptions CopyOptions;
+			CopyOptions.ApiVersion = EOS_PLAYERDATASTORAGE_COPYFILEMETADATAATINDEXOPTIONS_API_LATEST;
+			static_assert(EOS_PLAYERDATASTORAGE_COPYFILEMETADATAATINDEXOPTIONS_API_LATEST == 1, "EOS_PlayerDataStorage_CopyFileMetadataAtIndexOptions updated, check new fields");
+			CopyOptions.LocalUserId = LocalUserPuid;
+			CopyOptions.Index = FileIdx;
+
+			EOS_PlayerDataStorage_FileMetadata* FileMetadata = nullptr;
+			Result = EOS_PlayerDataStorage_CopyFileMetadataAtIndex(PlayerDataStorageHandle, &CopyOptions, &FileMetadata);
+			if (Result != EOS_EResult::EOS_Success)
+			{
+				UE_LOG(LogTemp, Warning, TEXT("EOS_PlayerDataStorage_CopyFileMetadataAtIndex failed with result=[%s]"), *LexToString(Result));
+				Op.SetError(FromEOSError(Result));
+				return;
+			}
+
+			NewEnumeratedFiles.Emplace(UTF8_TO_TCHAR(FileMetadata->Filename));
+
+			EOS_PlayerDataStorage_FileMetadata_Release(FileMetadata);
+		}
+
+		UserToFiles.Emplace(Params.LocalUserId, MoveTemp(NewEnumeratedFiles));
+		Op.SetResult({});
+	})
+	.Enqueue(GetSerialQueue());
+	return Op->GetHandle();
+}
+
+TOnlineResult<FUserFileGetEnumeratedFiles> FUserFileEOSGS::GetEnumeratedFiles(FUserFileGetEnumeratedFiles::Params&& Params)
+{
+	if (!Services.Get<FAuthEOSGS>()->IsLoggedIn(Params.LocalUserId))
+	{
+		return TOnlineResult<FUserFileGetEnumeratedFiles>(Errors::InvalidUser());
+	}
+
+	const TArray<FString>* Found = UserToFiles.Find(Params.LocalUserId);
+	if (!Found)
+	{
+		// Call EnumerateFiles first
+		return TOnlineResult<FUserFileGetEnumeratedFiles>(Errors::InvalidState());
+	}
+
+	return TOnlineResult<FUserFileGetEnumeratedFiles>({*Found});
+}
+
+TOnlineAsyncOpHandle<FUserFileReadFile> FUserFileEOSGS::ReadFile(FUserFileReadFile::Params&& InParams)
+{
+	static const TCHAR* FileContentsKey = TEXT("FileContents");
+
+	TOnlineAsyncOpRef<FUserFileReadFile> Op = GetOp<FUserFileReadFile>(MoveTemp(InParams));
+	const FUserFileReadFile::Params& Params = Op->GetParams();
+
+	if (Params.Filename.IsEmpty())
+	{
+		Op->SetError(Errors::InvalidParams());
+		return Op->GetHandle();
+	}
+
+	Op->Then([this](TOnlineAsyncOp<FUserFileReadFile>& Op, TPromise<const EOS_PlayerDataStorage_ReadFileCallbackInfo*>&& Promise)
+	{
+		const FUserFileReadFile::Params& Params = Op.GetParams();
+
+		if (!Services.Get<FAuthEOSGS>()->IsLoggedIn(Params.LocalUserId))
+		{
+			Op.SetError(Errors::InvalidUser());
+			Promise.EmplaceValue();
+			return;
+		}
+
+		FUserFileContentsRef FileContentsRef = MakeShared<FUserFileContents>();
+		Op.Data.Set<FUserFileContentsRef>(FileContentsKey, FileContentsRef);
+
+		using FReadFileDataCallback = TEOSGlobalCallback<EOS_PlayerDataStorage_OnReadFileDataCallback, EOS_PlayerDataStorage_ReadFileDataCallbackInfo, TOnlineAsyncOp<FUserFileReadFile>, EOS_PlayerDataStorage_EReadResult>;
+		TSharedRef<FReadFileDataCallback> ReadFileDataCallback = MakeShared<FReadFileDataCallback>(Op.AsWeak());
+		Op.Data.Set<TSharedRef<FReadFileDataCallback>>(TEXT("ReadFileDataCallback"), ReadFileDataCallback);
+		ReadFileDataCallback->CallbackLambda = [FileContentsRef](const EOS_PlayerDataStorage_ReadFileDataCallbackInfo* Data) -> EOS_PlayerDataStorage_EReadResult
+		{
+			FUserFileContents& FileContents = const_cast<FUserFileContents&>(FileContentsRef.Get());
+			// Only has any effect on the first callback
+			FileContents.Reserve(Data->TotalFileSizeBytes);
+			FileContents.Append((uint8*)Data->DataChunk, Data->DataChunkLengthBytes);
+
+			UE_LOG(LogTemp, VeryVerbose, TEXT("FUserFileEOSGS::ReadFile ReadProgress Filename=[%s] %d/%d"), UTF8_TO_TCHAR(Data->Filename), FileContents.Num(), Data->TotalFileSizeBytes);
+
+			return EOS_PlayerDataStorage_EReadResult::EOS_RR_ContinueReading;
+		};
+
+		const FTCHARToUTF8 Utf8Filename(Params.Filename);
+
+		EOS_PlayerDataStorage_ReadFileOptions  Options = {};
+		Options.ApiVersion = EOS_PLAYERDATASTORAGE_READFILEOPTIONS_API_LATEST;
+		static_assert(EOS_PLAYERDATASTORAGE_READFILEOPTIONS_API_LATEST == 1, "EOS_PlayerDataStorage_ReadFileOptions updated, check new fields");
+		Options.LocalUserId = GetProductUserIdChecked(Params.LocalUserId);
+		Options.Filename = Utf8Filename.Get();
+		Options.ReadChunkLengthBytes = Config.ChunkLengthBytes;
+		Options.ReadFileDataCallback = ReadFileDataCallback->GetCallbackPtr();
+		Options.FileTransferProgressCallback = &FUserFileEOSGS::OnFileTransferProgressStatic;
+
+		const EOS_HPlayerDataStorageFileTransferRequest RequestHandle = EOS_Async(EOS_PlayerDataStorage_ReadFile, PlayerDataStorageHandle, Options, MoveTemp(Promise));
+		Op.Data.Set<EOS_HPlayerDataStorageFileTransferRequest>(RequestHandleKey, RequestHandle);
+	})
+	.Then([this](TOnlineAsyncOp<FUserFileReadFile>& Op, const EOS_PlayerDataStorage_ReadFileCallbackInfo* Data)
+	{
+		if (Data->ResultCode == EOS_EResult::EOS_Success)
+		{
+			const FUserFileContentsRef* FileContents = Op.Data.Get<FUserFileContentsRef>(FileContentsKey);
+			if (ensure(FileContents))
+			{
+				Op.SetResult({ *FileContents });
+			}
+			else
+			{
+				Op.SetError(Errors::Unknown());
+			}
+		}
+		else
+		{
+			UE_LOG(LogTemp, Warning, TEXT("EOS_PlayerDataStorage_ReadFile failed with result=[%s]"), *LexToString(Data->ResultCode));
+			Op.SetError(FromEOSError(Data->ResultCode));
+		}
+
+		if (const EOS_HPlayerDataStorageFileTransferRequest* RequestHandle = Op.Data.Get<EOS_HPlayerDataStorageFileTransferRequest>(RequestHandleKey))
+		{
+			EOS_PlayerDataStorageFileTransferRequest_Release(*RequestHandle);
+		}
+	})
+	.Enqueue(GetSerialQueue());
+	return Op->GetHandle();
+}
+
+TOnlineAsyncOpHandle<FUserFileWriteFile> FUserFileEOSGS::WriteFile(FUserFileWriteFile::Params&& InParams)
+{
+	TOnlineAsyncOpRef<FUserFileWriteFile> Op = GetOp<FUserFileWriteFile>(MoveTemp(InParams));
+	const FUserFileWriteFile::Params& Params = Op->GetParams();
+
+	if (Params.Filename.IsEmpty())
+	{
+		Op->SetError(Errors::InvalidParams());
+		return Op->GetHandle();
+	}
+
+	Op->Then([this](TOnlineAsyncOp<FUserFileWriteFile>& Op, TPromise<const EOS_PlayerDataStorage_WriteFileCallbackInfo*>&& Promise)
+	{
+		const FUserFileWriteFile::Params& Params = Op.GetParams();
+
+		if (!Services.Get<FAuthEOSGS>()->IsLoggedIn(Params.LocalUserId))
+		{
+			Op.SetError(Errors::InvalidUser());
+			Promise.EmplaceValue();
+			return;
+		}
+
+		using FWriteFileDataCallback = TEOSGlobalCallback<EOS_PlayerDataStorage_OnWriteFileDataCallback, EOS_PlayerDataStorage_WriteFileDataCallbackInfo, TOnlineAsyncOp<FUserFileWriteFile>, EOS_PlayerDataStorage_EWriteResult, void*, uint32_t*>;
+		TSharedRef<FWriteFileDataCallback> WriteFileDataCallback = MakeShared<FWriteFileDataCallback>(Op.AsWeak());
+		Op.Data.Set<TSharedRef<FWriteFileDataCallback>>(TEXT("WriteFileDataCallback"), WriteFileDataCallback);
+		WriteFileDataCallback->CallbackLambda = [WeakOp = Op.AsWeak()] (const EOS_PlayerDataStorage_WriteFileDataCallbackInfo* CallbackInfo, void* OutDataBuffer, uint32_t* OutDataWritten)
+		{
+			TOnlineAsyncOpPtr<FUserFileWriteFile> Op = WeakOp.Pin();
+			if (!ensure(Op))
+			{
+				return EOS_PlayerDataStorage_EWriteResult::EOS_WR_FailRequest;
+			}
+
+			const TCHAR* BytesWrittenKey = TEXT("BytesWritten");
+
+			const FUserFileWriteFile::Params& Params = Op->GetParams();
+			const FUserFileContents& FileContents = Params.FileContents;
+
+			const uint32_t* BytesWrittenPtr = Op->Data.Get<uint32_t>(BytesWrittenKey);
+			const uint32_t BytesWrittenBefore = BytesWrittenPtr ? *BytesWrittenPtr : 0;
+			const uint32_t BytesRemaining = FileContents.Num() - BytesWrittenBefore;
+			const uint32_t BytesToWrite = FMath::Min(CallbackInfo->DataBufferLengthBytes, BytesRemaining);
+
+			FMemory::Memcpy(OutDataBuffer, FileContents.GetData() + BytesWrittenBefore, BytesToWrite);
+			*OutDataWritten = BytesToWrite;
+
+			const uint32_t BytesWrittenTotal = BytesWrittenBefore + BytesToWrite;
+			Op->Data.Set<uint32_t>(BytesWrittenKey, BytesWrittenTotal);
+
+			UE_LOG(LogTemp, VeryVerbose, TEXT("FUserFileEOSGS::WriteFile WriteProgress Filename=[%s] %d/%d"), UTF8_TO_TCHAR(CallbackInfo->Filename), BytesWrittenTotal, FileContents.Num());
+
+			const bool bDone = BytesWrittenTotal == FileContents.Num();
+			if (bDone)
+			{
+				return EOS_PlayerDataStorage_EWriteResult::EOS_WR_CompleteRequest;
+			}
+			else
+			{
+				return EOS_PlayerDataStorage_EWriteResult::EOS_WR_ContinueWriting;
+			}
+		};
+
+		const FTCHARToUTF8 Utf8Filename(Params.Filename);
+
+		EOS_PlayerDataStorage_WriteFileOptions  Options = {};
+		Options.ApiVersion = EOS_PLAYERDATASTORAGE_WRITEFILEOPTIONS_API_LATEST;
+		static_assert(EOS_PLAYERDATASTORAGE_WRITEFILEOPTIONS_API_LATEST == 1, "EOS_PlayerDataStorage_WriteFileOptions updated, check new fields");
+		Options.LocalUserId = GetProductUserIdChecked(Params.LocalUserId);
+		Options.Filename = Utf8Filename.Get();
+		Options.ChunkLengthBytes = Config.ChunkLengthBytes;
+		Options.WriteFileDataCallback = WriteFileDataCallback->GetCallbackPtr();
+		Options.FileTransferProgressCallback = &FUserFileEOSGS::OnFileTransferProgressStatic;
+
+		const EOS_HPlayerDataStorageFileTransferRequest RequestHandle = EOS_Async(EOS_PlayerDataStorage_WriteFile, PlayerDataStorageHandle, Options, MoveTemp(Promise));
+		Op.Data.Set<EOS_HPlayerDataStorageFileTransferRequest>(RequestHandleKey, RequestHandle);
+	})
+	.Then([this](TOnlineAsyncOp<FUserFileWriteFile>& Op, const EOS_PlayerDataStorage_WriteFileCallbackInfo* Data)
+	{
+		if (Data->ResultCode == EOS_EResult::EOS_Success)
+		{
+			Op.SetResult({});
+		}
+		else
+		{
+			UE_LOG(LogTemp, Warning, TEXT("EOS_PlayerDataStorage_WriteFile failed with result=[%s]"), *LexToString(Data->ResultCode));
+			Op.SetError(FromEOSError(Data->ResultCode));
+		}
+
+		if (const EOS_HPlayerDataStorageFileTransferRequest* RequestHandle = Op.Data.Get<EOS_HPlayerDataStorageFileTransferRequest>(RequestHandleKey))
+		{
+			EOS_PlayerDataStorageFileTransferRequest_Release(*RequestHandle);
+		}
+	})
+	.Enqueue(GetSerialQueue());
+	return Op->GetHandle();
+}
+
+TOnlineAsyncOpHandle<FUserFileCopyFile> FUserFileEOSGS::CopyFile(FUserFileCopyFile::Params&& InParams)
+{
+	TOnlineAsyncOpRef<FUserFileCopyFile> Op = GetOp<FUserFileCopyFile>(MoveTemp(InParams));
+	const FUserFileCopyFile::Params& Params = Op->GetParams();
+
+	if (Params.SourceFilename.IsEmpty() || Params.TargetFilename.IsEmpty())
+	{
+		Op->SetError(Errors::InvalidParams());
+		return Op->GetHandle();
+	}
+
+	Op->Then([this](TOnlineAsyncOp<FUserFileCopyFile>& Op, TPromise<const EOS_PlayerDataStorage_DuplicateFileCallbackInfo*>&& Promise)
+	{
+		const FUserFileCopyFile::Params& Params = Op.GetParams();
+
+		if (!Services.Get<FAuthEOSGS>()->IsLoggedIn(Params.LocalUserId))
+		{
+			Op.SetError(Errors::InvalidUser());
+			Promise.EmplaceValue();
+			return;
+		}
+
+		const FTCHARToUTF8 Utf8SourceFilename(Params.SourceFilename);
+		const FTCHARToUTF8 Utf8TargetFilename(Params.TargetFilename);
+
+		EOS_PlayerDataStorage_DuplicateFileOptions  Options = {};
+		Options.ApiVersion = EOS_PLAYERDATASTORAGE_DUPLICATEFILEOPTIONS_API_LATEST;
+		static_assert(EOS_PLAYERDATASTORAGE_DUPLICATEFILEOPTIONS_API_LATEST == 1, "EOS_PlayerDataStorage_DuplicateFileOptions updated, check new fields");
+		Options.LocalUserId = GetProductUserIdChecked(Params.LocalUserId);
+		Options.SourceFilename = Utf8SourceFilename.Get();
+		Options.DestinationFilename = Utf8TargetFilename.Get();
+
+		EOS_Async(EOS_PlayerDataStorage_DuplicateFile, PlayerDataStorageHandle, Options, MoveTemp(Promise));
+	})
+	.Then([this](TOnlineAsyncOp<FUserFileCopyFile>& Op, const EOS_PlayerDataStorage_DuplicateFileCallbackInfo* Data)
+	{
+		if (Data->ResultCode == EOS_EResult::EOS_Success)
+		{
+			Op.SetResult({});
+		}
+		else
+		{
+			UE_LOG(LogTemp, Warning, TEXT("EOS_PlayerDataStorage_DuplicateFile failed with result=[%s]"), *LexToString(Data->ResultCode));
+			Op.SetError(FromEOSError(Data->ResultCode));
+		}
+	})
+	.Enqueue(GetSerialQueue());
+	return Op->GetHandle();
+}
+
+TOnlineAsyncOpHandle<FUserFileDeleteFile> FUserFileEOSGS::DeleteFile(FUserFileDeleteFile::Params&& InParams)
+{
+	TOnlineAsyncOpRef<FUserFileDeleteFile> Op = GetOp<FUserFileDeleteFile>(MoveTemp(InParams));
+	const FUserFileDeleteFile::Params& Params = Op->GetParams();
+
+	if (Params.Filename.IsEmpty())
+	{
+		Op->SetError(Errors::InvalidParams());
+		return Op->GetHandle();
+	}
+
+	Op->Then([this](TOnlineAsyncOp<FUserFileDeleteFile>& Op, TPromise<const EOS_PlayerDataStorage_DeleteFileCallbackInfo*>&& Promise)
+	{
+		const FUserFileDeleteFile::Params& Params = Op.GetParams();
+
+		if (!Services.Get<FAuthEOSGS>()->IsLoggedIn(Params.LocalUserId))
+		{
+			Op.SetError(Errors::InvalidUser());
+			Promise.EmplaceValue();
+			return;
+		}
+
+		const FTCHARToUTF8 Utf8Filename(Params.Filename);
+
+		EOS_PlayerDataStorage_DeleteFileOptions  Options = {};
+		Options.ApiVersion = EOS_PLAYERDATASTORAGE_DELETEFILEOPTIONS_API_LATEST;
+		static_assert(EOS_PLAYERDATASTORAGE_DELETEFILEOPTIONS_API_LATEST == 1, "EOS_PlayerDataStorage_DeleteFileOptions updated, check new fields");
+		Options.LocalUserId = GetProductUserIdChecked(Params.LocalUserId);
+		Options.Filename = Utf8Filename.Get();
+
+		EOS_Async(EOS_PlayerDataStorage_DeleteFile, PlayerDataStorageHandle, Options, MoveTemp(Promise));
+	})
+	.Then([this](TOnlineAsyncOp<FUserFileDeleteFile>& Op, const EOS_PlayerDataStorage_DeleteFileCallbackInfo* Data)
+	{
+		if (Data->ResultCode == EOS_EResult::EOS_Success)
+		{
+			Op.SetResult({});
+		}
+		else
+		{
+			UE_LOG(LogTemp, Warning, TEXT("EOS_PlayerDataStorage_DeleteFile failed with result=[%s]"), *LexToString(Data->ResultCode));
+			Op.SetError(FromEOSError(Data->ResultCode));
+		}
+	})
+	.Enqueue(GetSerialQueue());
+	return Op->GetHandle();
+}
+
+
+void EOS_CALL FUserFileEOSGS::OnFileTransferProgressStatic(const EOS_PlayerDataStorage_FileTransferProgressCallbackInfo* Data)
+{
+	UE_LOG(LogTemp, VeryVerbose, TEXT("FUserFileEOSGS::ReadFile TransferProgress Filename=[%s] %d/%d"), UTF8_TO_TCHAR(Data->Filename), Data->BytesTransferred, Data->TotalFileSizeBytes);
+}
+
+/* UE::Online */ }
diff --git a/Engine/Plugins/Online/Experimental/OnlineServicesEOSGS/Source/Public/Online/UserFileEOSGS.h b/Engine/Plugins/Online/Experimental/OnlineServicesEOSGS/Source/Public/Online/UserFileEOSGS.h
new file mode 100644
--- /dev/null
+++ b/Engine/Plugins/Online/Experimental/OnlineServicesEOSGS/Source/Public/Online/UserFileEOSGS.h
@@ -0,0 +1,60 @@
+// Copyright Epic Games, Inc. All Rights Reserved.
+
+#pragma once
+
+#include "CoreMinimal.h"
+#include "Online/UserFileCommon.h"
+
+#if defined(EOS_PLATFORM_BASE_FILE_NAME)
+#include EOS_PLATFORM_BASE_FILE_NAME
+#endif
+#include "eos_playerdatastorage_types.h"
+
+namespace UE::Online {
+
+class FOnlineServicesEOSGS;
+
+struct FUserFileEOSGSConfig
+{
+	int32 ChunkLengthBytes = 4096;
+};
+
+namespace Meta {
+
+BEGIN_ONLINE_STRUCT_META(FUserFileEOSGSConfig)
+	ONLINE_STRUCT_FIELD(FUserFileEOSGSConfig, ChunkLengthBytes)
+END_ONLINE_STRUCT_META()
+
+/* Meta */ }
+
+class ONLINESERVICESEOSGS_API FUserFileEOSGS : public FUserFileCommon
+{
+public:
+	using Super = FUserFileCommon;
+
+	FUserFileEOSGS(FOnlineServicesEOSGS& InOwningSubsystem);
+	virtual ~FUserFileEOSGS() = default;
+
+	// IOnlineComponent
+	virtual void Initialize() override;
+	virtual void LoadConfig() override;
+
+	// IUserFile
+	virtual TOnlineAsyncOpHandle<FUserFileEnumerateFiles> EnumerateFiles(FUserFileEnumerateFiles::Params&& Params) override;
+	virtual TOnlineResult<FUserFileGetEnumeratedFiles> GetEnumeratedFiles(FUserFileGetEnumeratedFiles::Params&& Params) override;
+	virtual TOnlineAsyncOpHandle<FUserFileReadFile> ReadFile(FUserFileReadFile::Params&& Params) override;
+	virtual TOnlineAsyncOpHandle<FUserFileWriteFile> WriteFile(FUserFileWriteFile::Params&& Params) override;
+	virtual TOnlineAsyncOpHandle<FUserFileCopyFile> CopyFile(FUserFileCopyFile::Params&& Params) override;
+	virtual TOnlineAsyncOpHandle<FUserFileDeleteFile> DeleteFile(FUserFileDeleteFile::Params&& Params) override;
+
+protected:
+	EOS_HPlayerDataStorage PlayerDataStorageHandle = nullptr;
+
+	FUserFileEOSGSConfig Config;
+
+	TMap<FOnlineAccountIdHandle, TArray<FString>> UserToFiles;
+
+	static void EOS_CALL OnFileTransferProgressStatic(const EOS_PlayerDataStorage_FileTransferProgressCallbackInfo* Data);
+};
+
+/* UE::Online */ }
