- Fix ensure in AchievementsEOSGS due to broken ApiVersion sanity checks.

diff --git a/Engine/Plugins/Online/Experimental/OnlineServicesEOSGS/Source/Private/Online/AchievementsEOSGS.cpp b/Engine/Plugins/Online/Experimental/OnlineServicesEOSGS/Source/Private/Online/AchievementsEOSGS.cpp
--- a/Engine/Plugins/Online/Experimental/OnlineServicesEOSGS/Source/Private/Online/AchievementsEOSGS.cpp
+++ b/Engine/Plugins/Online/Experimental/OnlineServicesEOSGS/Source/Private/Online/AchievementsEOSGS.cpp
@@ -112,10 +112,8 @@ TOnlineAsyncOpHandle<FQueryAchievementDefinitions> FAchievementsEOSGS::QueryAchi
 				}
 
 				static_assert(EOS_ACHIEVEMENTS_DEFINITIONV2_API_LATEST == 2, "EOS_Achievements_DefinitionV2 updated, check new fields");
-				if (!ensure(EOSDefinition->ApiVersion == EOS_ACHIEVEMENTS_DEFINITIONV2_API_LATEST))
-				{
-					UE_LOG(LogTemp, Warning, TEXT("EOS_Achievements_DefinitionV2 version mismatch Expected=%d Actual=%d"), EOS_ACHIEVEMENTS_DEFINITIONV2_API_LATEST, EOSDefinition->ApiVersion);
-				}
+				const bool bDefsApiVersionOk = EOSDefinition->ApiVersion == EOS_ACHIEVEMENTS_DEFINITIONV2_API_LATEST;
+				UE_CLOG(!bDefsApiVersionOk, LogTemp, Warning, TEXT("EOS_Achievements_DefinitionV2 version mismatch Expected=%d Actual=%d"), EOS_ACHIEVEMENTS_DEFINITIONV2_API_LATEST, EOSDefinition->ApiVersion);
 
 				FString AchievementId = UTF8_TO_TCHAR(EOSDefinition->AchievementId);
 				FAchievementDefinition& AchievementDefinition = NewAchievementDefinitions.Emplace(AchievementId);
@@ -130,14 +128,12 @@ TOnlineAsyncOpHandle<FQueryAchievementDefinitions> FAchievementsEOSGS::QueryAchi
 				AchievementDefinition.bIsHidden = EOSDefinition->bIsHidden == EOS_TRUE;
 				for (uint32_t StatThresholdsIdx = 0; StatThresholdsIdx < EOSDefinition->StatThresholdsCount; StatThresholdsIdx++)
 				{
-					static_assert(EOS_ACHIEVEMENTS_STATTHRESHOLDS_API_LATEST == 1, "EOS_Achievements_StatThresholds updated, check new fields");
-					if (!ensure(EOSDefinition->ApiVersion == EOS_ACHIEVEMENTS_STATTHRESHOLDS_API_LATEST))
-					{
-						UE_LOG(LogTemp, Warning, TEXT("EOS_Achievements_StatThresholds version mismatch Expected=%d Actual=%d"), EOS_ACHIEVEMENTS_DEFINITIONV2_API_LATEST, EOSDefinition->ApiVersion);
-					}
-
 					const EOS_Achievements_StatThresholds& EOSStatThreshold = EOSDefinition->StatThresholds[StatThresholdsIdx];
 
+					static_assert(EOS_ACHIEVEMENTS_STATTHRESHOLDS_API_LATEST == 1, "EOS_Achievements_StatThresholds updated, check new fields");
+					const bool bStatsApiVersionOk = EOSStatThreshold.ApiVersion == EOS_ACHIEVEMENTS_STATTHRESHOLDS_API_LATEST;
+					UE_CLOG(!bStatsApiVersionOk, LogTemp, Warning, TEXT("EOS_Achievements_StatThresholds version mismatch Expected=%d Actual=%d"), EOS_ACHIEVEMENTS_STATTHRESHOLDS_API_LATEST, EOSStatThreshold.ApiVersion);
+
 					FAchievementStatDefinition& StatDefinition = AchievementDefinition.StatDefinitions.Emplace_GetRef();
 					StatDefinition.StatId = UTF8_TO_TCHAR(EOSStatThreshold.Name);
 					StatDefinition.UnlockThreshold = uint32_t(EOSStatThreshold.Threshold);
