Shutdown socket subsystem on shutdown

diff --git a/Engine/Plugins/Online/Experimental/OnlineServicesEOSGS/Source/Private/Online/OnlineServicesEOSGS.cpp b/Engine/Plugins/Online/Experimental/OnlineServicesEOSGS/Source/Private/Online/OnlineServicesEOSGS.cpp
--- a/Engine/Plugins/Online/Experimental/OnlineServicesEOSGS/Source/Private/Online/OnlineServicesEOSGS.cpp
+++ b/Engine/Plugins/Online/Experimental/OnlineServicesEOSGS/Source/Private/Online/OnlineServicesEOSGS.cpp
@@ -68,6 +68,18 @@ void FOnlineServicesEOSGS::Initialize()
 	Super::Initialize();
 }
 
+void FOnlineServicesEOSGS::Shutdown()
+{
+	Super::Shutdown();
+#if WITH_ENGINE
+	if (SocketSubsystem)
+	{
+		SocketSubsystem->Shutdown();
+		SocketSubsystem = nullptr;
+	}
+#endif
+}
+
 TOnlineResult<FGetResolvedConnectString> FOnlineServicesEOSGS::GetResolvedConnectString(FGetResolvedConnectString::Params&& Params)
 {
 	if (Params.LobbyId.IsValid())
diff --git a/Engine/Plugins/Online/Experimental/OnlineServicesEOSGS/Source/Public/Online/OnlineServicesEOSGS.h b/Engine/Plugins/Online/Experimental/OnlineServicesEOSGS/Source/Public/Online/OnlineServicesEOSGS.h
--- a/Engine/Plugins/Online/Experimental/OnlineServicesEOSGS/Source/Public/Online/OnlineServicesEOSGS.h
+++ b/Engine/Plugins/Online/Experimental/OnlineServicesEOSGS/Source/Public/Online/OnlineServicesEOSGS.h
@@ -27,6 +27,7 @@ public:
 
 	virtual void RegisterComponents() override;
 	virtual void Initialize() override;
+	virtual void Shutdown() override;
 	virtual TOnlineResult<FGetResolvedConnectString> GetResolvedConnectString(FGetResolvedConnectString::Params&& Params) override;
 	virtual EOnlineServices GetServicesProvider() const override { return EOnlineServices::Epic; }
 
