Auth EOS/GS implementation

diff --git a/Engine/Plugins/Online/Experimental/OnlineServicesEOS/Source/Private/Online/AuthEOS.cpp b/Engine/Plugins/Online/Experimental/OnlineServicesEOS/Source/Private/Online/AuthEOS.cpp
--- a/Engine/Plugins/Online/Experimental/OnlineServicesEOS/Source/Private/Online/AuthEOS.cpp
+++ b/Engine/Plugins/Online/Experimental/OnlineServicesEOS/Source/Private/Online/AuthEOS.cpp
@@ -22,19 +22,27 @@
 
 namespace UE::Online {
 
-struct FAuthEOSConfig
+struct FAuthEOSLoginConfig
 {
-	FString DefaultExternalCredentialTypeStr;
+	TArray<FString> DefaultScopes;
 };
 
 namespace Meta {
 
-BEGIN_ONLINE_STRUCT_META(FAuthEOSConfig)
-	ONLINE_STRUCT_FIELD(FAuthEOSConfig, DefaultExternalCredentialTypeStr)
+BEGIN_ONLINE_STRUCT_META(FAuthEOSLoginConfig)
+	ONLINE_STRUCT_FIELD(FAuthEOSLoginConfig, DefaultScopes)
 END_ONLINE_STRUCT_META()
 
 /* Meta*/ }
 
+
+namespace
+{
+
+static const FString AccountInfoKeyName = TEXT("AccountInfoEOS");
+
+/* anonymous */ }
+
 FAuthEOS::FAuthEOS(FOnlineServicesEOS& InServices)
 	: Super(InServices)
 {
@@ -44,466 +52,297 @@ void FAuthEOS::Initialize()
 {
 	Super::Initialize();
 
-	// Register for login status changes
-	EOS_Auth_AddNotifyLoginStatusChangedOptions Options = { };
-	Options.ApiVersion = EOS_AUTH_ADDNOTIFYLOGINSTATUSCHANGED_API_LATEST;
-	NotifyEASLoginStatusChangedNotificationId = EOS_Auth_AddNotifyLoginStatusChanged(AuthHandle, &Options, this, [](const EOS_Auth_LoginStatusChangedCallbackInfo* Data)
-	{
-		FAuthEOS* This = reinterpret_cast<FAuthEOS*>(Data->ClientData);
-
-		FOnlineAccountIdHandle LocalUserId = FindAccountId(Data->LocalUserId);
-		// invalid handle is expected for players logging in because this callback is called _before_ the login complete callback
-		if (LocalUserId.IsValid())
-		{
-			ELoginStatus PreviousStatus = ToELoginStatus(Data->PrevStatus);
-			ELoginStatus CurrentStatus = ToELoginStatus(Data->CurrentStatus);
-			This->OnEASLoginStatusChanged(LocalUserId, PreviousStatus, CurrentStatus);
-		}
-	});
+	UserInfoHandle = EOS_Platform_GetUserInfoInterface(static_cast<FOnlineServicesEOS&>(GetServices()).GetEOSPlatformHandle());
+	check(UserInfoHandle != nullptr);
 }
 
 TOnlineAsyncOpHandle<FAuthLogin> FAuthEOS::Login(FAuthLogin::Params&& Params)
 {
-	// Is this auto-login?
-	if (Params.CredentialsId.IsEmpty() && Params.CredentialsType.IsEmpty() && Params.CredentialsToken.IsType<FString>() && Params.CredentialsToken.Get<FString>().IsEmpty())
-	{
-		FString	CommandLineAuthId;
-		FString	CommandLineAuthToken;
-		FString	CommandLineAuthType;
-		FParse::Value(FCommandLine::Get(), TEXT("AUTH_LOGIN="), CommandLineAuthId);
-		FParse::Value(FCommandLine::Get(), TEXT("AUTH_PASSWORD="), CommandLineAuthToken);
-		FParse::Value(FCommandLine::Get(), TEXT("AUTH_TYPE="), CommandLineAuthType);
-		if (!CommandLineAuthId.IsEmpty() && !CommandLineAuthToken.IsEmpty() && !CommandLineAuthType.IsEmpty())
-		{
-			Params.CredentialsId = MoveTemp(CommandLineAuthId);
-			Params.CredentialsToken.Emplace<FString>(MoveTemp(CommandLineAuthToken));
-			Params.CredentialsType = MoveTemp(CommandLineAuthType);
-		}
-	}
 	TOnlineAsyncOpRef<FAuthLogin> Op = GetOp<FAuthLogin>(MoveTemp(Params));
-	// Are we already logged in?
-	if (GetAccountIdByPlatformUserId(Op->GetParams().PlatformUserId).IsOk())
-	{
-		Op->SetError(Errors::Auth::AlreadyLoggedIn());
-		return Op->GetHandle();
-	}
-
-	TOnlineChainableAsyncOp<FAuthLogin, TSharedPtr<FEOSConnectLoginCredentials>> ConnectLoginOp = LoginEAS(*Op);
-
-	// Check if LoginEAS completed (failed) the operation
-	if (ConnectLoginOp.GetOwningOperation().IsComplete())
-	{
-		return ConnectLoginOp.GetOwningOperation().GetHandle();
-	}
-
-	ConnectLoginOp.Then([this](TOnlineAsyncOp<FAuthLogin>& InAsyncOp, TSharedPtr<FEOSConnectLoginCredentials>&& InConnectLoginCredentials, TPromise<const EOS_Connect_LoginCallbackInfo*>&& Promise)
-	{
-		EOS_Connect_LoginOptions ConnectLoginOptions = { };
-		ConnectLoginOptions.ApiVersion = EOS_CONNECT_LOGIN_API_LATEST;
-		ConnectLoginOptions.Credentials = InConnectLoginCredentials.Get();
-
-		EOS_Async(EOS_Connect_Login, ConnectHandle, ConnectLoginOptions, MoveTemp(Promise));
-	})
-	.Then([this](TOnlineAsyncOp<FAuthLogin>& InAsyncOp, const EOS_Connect_LoginCallbackInfo* Data, TPromise<const EOS_Connect_CreateUserCallbackInfo*>&& Promise)
+	// Step 1: Set up operation data.
+	Op->Then([this](TOnlineAsyncOp<FAuthLogin>& InAsyncOp)
 	{
-		UE_LOG(LogTemp, Verbose, TEXT("[FAuthEOS::Login] EOS_Connect_Login Result: [%s]"), *LexToString(Data->ResultCode));
-
-		if (Data->ResultCode == EOS_EResult::EOS_Success)
-		{
-			// We cache the Product User Id to use it in later stages of the login process
-			InAsyncOp.Data.Set(TEXT("ProductUserId"), Data->LocalUserId);
-
-			ProcessSuccessfulLogin(InAsyncOp);
-		}
-		else if (Data->ResultCode == EOS_EResult::EOS_InvalidUser && Data->ContinuanceToken != nullptr)
+		const FAuthLogin::Params& Params = InAsyncOp.GetParams();
+		TSharedPtr<FAccountInfoEOS> AccountInfoEOS = AccountInfoRegistryEOS.Find(Params.PlatformUserId);
+		if (AccountInfoEOS)
 		{
-			EOS_Connect_CreateUserOptions ConnectCreateUserOptions = { };
-			ConnectCreateUserOptions.ApiVersion = EOS_CONNECT_CREATEUSER_API_LATEST;
-			ConnectCreateUserOptions.ContinuanceToken = Data->ContinuanceToken;
-
-			EOS_Async(EOS_Connect_CreateUser, ConnectHandle, ConnectCreateUserOptions, MoveTemp(Promise));
+			InAsyncOp.SetError(Errors::Auth::AlreadyLoggedIn());
 			return;
 		}
-		else
-		{
-			// TODO: EAS Logout
-			InAsyncOp.SetError(Errors::Unknown()); // TODO
-		}
-
-		Promise.EmplaceValue();
-	})
-	.Then([this](TOnlineAsyncOp<FAuthLogin>& InAsyncOp, const EOS_Connect_CreateUserCallbackInfo* Data)
-	{
-		UE_LOG(LogTemp, Verbose, TEXT("[FAuthEOS::Login] EOS_Connect_CreateUser Result: [%s]"), *LexToString(Data->ResultCode));
-
-		if (Data->ResultCode == EOS_EResult::EOS_Success)
-		{
-			// We cache the Product User Id to use it in later stages of the login process
-			InAsyncOp.Data.Set(TEXT("ProductUserId"), Data->LocalUserId);
 
-			ProcessSuccessfulLogin(InAsyncOp);
-		}
-		else
-		{
-			// TODO: EAS Logout
+		AccountInfoEOS = MakeShared<FAccountInfoEOS>();
+		AccountInfoEOS->PlatformUserId = Params.PlatformUserId;
+		AccountInfoEOS->LoginStatus = ELoginStatus::NotLoggedIn;
 
-			InAsyncOp.SetError(Errors::Unknown()); // TODO
-		}
+		// Set user auth data on operation.
+		InAsyncOp.Data.Set<TSharedRef<FAccountInfoEOS>>(AccountInfoKeyName, AccountInfoEOS.ToSharedRef());
 	})
-	.Enqueue(GetSerialQueue());
-
-	return Op->GetHandle();
-}
-
-TOnlineChainableAsyncOp<FAuthLogin, TSharedPtr<FEOSConnectLoginCredentials>> FAuthEOS::LoginEAS(TOnlineAsyncOp<FAuthLogin>& InAsyncOp)
-{
-	const FAuthLogin::Params& Params = InAsyncOp.GetParams();
-	auto FailureLambda = [](TOnlineAsyncOp<FAuthLogin>& Op) -> TSharedPtr<FEOSConnectLoginCredentials>
+	// Step 2: Login EAS.
+	.Then([this](TOnlineAsyncOp<FAuthLogin>& InAsyncOp)
 	{
-		return nullptr;
-	};
+		TPromise<void> Promise;
+		TFuture<void> Future = Promise.GetFuture();
 
-	EOS_Auth_LoginOptions LoginOptions = { };
-	LoginOptions.ApiVersion = EOS_AUTH_LOGIN_API_LATEST;
-	bool bContainsFlagsNone = false;
-	for (const FString& Scope : Params.Scopes)
-	{
-		EOS_EAuthScopeFlags ScopeFlag;
-		if (LexFromString(ScopeFlag, *Scope))
-		{
-			if (ScopeFlag == EOS_EAuthScopeFlags::EOS_AS_NoFlags)
-			{
-				bContainsFlagsNone = true;
-			}
-			LoginOptions.ScopeFlags |= ScopeFlag;
-		}
-		else
-		{
-			UE_LOG(LogTemp, Warning, TEXT("Invalid ScopeFlag=[%s]"), *Scope);
-			InAsyncOp.SetError(Errors::Unknown());
-			return InAsyncOp.Then(FailureLambda);
-		}
-	}
-	// TODO:  Where to put default scopes?
-	if (!bContainsFlagsNone && LoginOptions.ScopeFlags == EOS_EAuthScopeFlags::EOS_AS_NoFlags)
-	{
-		LoginOptions.ScopeFlags = EOS_EAuthScopeFlags::EOS_AS_BasicProfile | EOS_EAuthScopeFlags::EOS_AS_FriendsList | EOS_EAuthScopeFlags::EOS_AS_Presence;
-	}
+		const FAuthLogin::Params& Params = InAsyncOp.GetParams();
 
-	// First, we'll check if any subtypes have been specified
-	FString CredentialTypeStr;
-	FString ExternalCredentialTypeStr;
-	Params.CredentialsType.Split(FString(TEXT(":")), &CredentialTypeStr, &ExternalCredentialTypeStr);
+		FAuthEOSLoginConfig AuthEOSLoginConfig;
+		LoadConfig(AuthEOSLoginConfig, TEXT("Login"));
 
-	// If no subtype was specified, we'll treat the string passed as the type
-	if (CredentialTypeStr.IsEmpty())
-	{
-		CredentialTypeStr = Params.CredentialsType;
-	}
+		FLoginEASImpl::Params LoginParams;
+		LoginParams.PlatformUserId = Params.PlatformUserId;
+		LoginParams.CredentialsType = Params.CredentialsType;
+		LoginParams.CredentialsId = Params.CredentialsId;
+		LoginParams.CredentialsToken = Params.CredentialsToken;
+		LoginParams.Scopes = !Params.Scopes.IsEmpty() ? Params.Scopes : AuthEOSLoginConfig.DefaultScopes;
 
-	FEOSAuthCredentials Credentials;
-	if (LexFromString(Credentials.Type, *CredentialTypeStr))
-	{
-		switch (Credentials.Type)
+		LoginEASImpl(LoginParams)
+		.Next([this, Promise = MoveTemp(Promise), WeakOp = InAsyncOp.AsWeak()](TDefaultErrorResult<FLoginEASImpl>&& LoginResult) mutable -> void
 		{
-		case EOS_ELoginCredentialType::EOS_LCT_ExternalAuth:
-			// If an external credential type wasn't specified, we'll grab the platform default from the configuration values
-			if (ExternalCredentialTypeStr.IsEmpty())
+			if (TSharedPtr<TOnlineAsyncOp<FAuthLogin>> Op = WeakOp.Pin())
 			{
-				FAuthEOSConfig AuthEOSConfig;
-				LoadConfig(AuthEOSConfig);
-				ExternalCredentialTypeStr = AuthEOSConfig.DefaultExternalCredentialTypeStr;
-			}
+				const TSharedRef<FAccountInfoEOS>& AccountInfoEOS = GetOpDataChecked<TSharedRef<FAccountInfoEOS>>(*Op, AccountInfoKeyName);
 
-			EOS_EExternalCredentialType ExternalType;
-			if (LexFromString(ExternalType, *ExternalCredentialTypeStr))
-			{
-				Credentials.ExternalType = ExternalType;
+				if (LoginResult.IsError())
+				{
+					UE_LOG(LogOnlineServices, Warning, TEXT("[FAuthEOS::Login] Failure: LoginEASImpl %s"), *LoginResult.GetErrorValue().GetLogString());
+					Op->SetError(Errors::Unknown(MoveTemp(LoginResult.GetErrorValue())));
+				}
+				else
+				{
+					// Cache EpicAccountId on successful EAS login.
+					AccountInfoEOS->EpicAccountId = LoginResult.GetOkValue().EpicAccountId;
+				}
 			}
 
-			Credentials.SetToken(Params.CredentialsToken);
-			break;
-		case EOS_ELoginCredentialType::EOS_LCT_ExchangeCode:
-			// This is how the Epic launcher will pass credentials to you
-			Credentials.IdAnsi[0] = '\0';
-			FCStringAnsi::Strncpy(Credentials.TokenAnsi, TCHAR_TO_UTF8(*Params.CredentialsToken.Get<FString>()), EOS_MAX_TOKEN_SIZE);
-			break;
-		case EOS_ELoginCredentialType::EOS_LCT_Password:
-			FCStringAnsi::Strncpy(Credentials.IdAnsi, TCHAR_TO_UTF8(*Params.CredentialsId), EOS_STRING_BUFFER_LENGTH);
-			FCStringAnsi::Strncpy(Credentials.TokenAnsi, TCHAR_TO_UTF8(*Params.CredentialsToken.Get<FString>()), EOS_MAX_TOKEN_SIZE);
-			break;
-		case EOS_ELoginCredentialType::EOS_LCT_Developer:
-			// This is auth via the EOS auth tool
-			FCStringAnsi::Strncpy(Credentials.IdAnsi, TCHAR_TO_UTF8(*Params.CredentialsId), EOS_STRING_BUFFER_LENGTH);
-			FCStringAnsi::Strncpy(Credentials.TokenAnsi, TCHAR_TO_UTF8(*Params.CredentialsToken.Get<FString>()), EOS_MAX_TOKEN_SIZE);
-			break;
-		case EOS_ELoginCredentialType::EOS_LCT_AccountPortal:
-			// This is auth via the EOS Account Portal
-			Credentials.IdAnsi[0] = '\0';
-			Credentials.TokenAnsi[0] = '\0';
-			break;
-		case EOS_ELoginCredentialType::EOS_LCT_PersistentAuth:
-			// This is auth via stored credentials in EOS
-			Credentials.Id = nullptr;
-			Credentials.Token = nullptr;
-			break;
-		default:
-			UE_LOG(LogTemp, Warning, TEXT("Unsupported CredentialsType=[%s]"), *Params.CredentialsType);
-			InAsyncOp.SetError(Errors::Unknown()); // TODO
-			return InAsyncOp.Then(FailureLambda);
-		}
-	}
-	else
-	{
-		UE_LOG(LogTemp, Warning, TEXT("Invalid CredentialsType=[%s]"), *Params.CredentialsType);
-		InAsyncOp.SetError(Errors::Unknown()); // TODO
-		return InAsyncOp.Then(FailureLambda);
-	}
+			Promise.EmplaceValue();
+		});
 
-	return InAsyncOp.Then([this, LoginOptions, Credentials](TOnlineAsyncOp<FAuthLogin>& InAsyncOp, TPromise<const EOS_Auth_LoginCallbackInfo*>&& Promise) mutable
-	{
-		LoginOptions.Credentials = &Credentials;
-		EOS_Async(EOS_Auth_Login, AuthHandle, LoginOptions, MoveTemp(Promise));
+		return Future;
 	})
-	.Then([this](TOnlineAsyncOp<FAuthLogin>& InAsyncOp, const EOS_Auth_LoginCallbackInfo* Data, TPromise<const EOS_Auth_LinkAccountCallbackInfo*>&& Promise)
+	// Step 3: Fetch external auth credentials for connect login.
+	.Then([this](TOnlineAsyncOp<FAuthLogin>& InAsyncOp)
 	{
-		UE_LOG(LogTemp, Verbose, TEXT("[FAuthEOS::Login] EOS_Auth_Login Result: [%s]"), *LexToString(Data->ResultCode));
+		const FAuthLogin::Params& Params = InAsyncOp.GetParams();
+		const TSharedRef<FAccountInfoEOS>& AccountInfoEOS = GetOpDataChecked<TSharedRef<FAccountInfoEOS>>(InAsyncOp, AccountInfoKeyName);
 
-		if (Data->ResultCode == EOS_EResult::EOS_Success)
-		{
-			// We cache the Epic Account Id to use it in later stages of the login process
-			InAsyncOp.Data.Set(TEXT("EpicAccountId"), Data->LocalUserId);
+		TPromise<FLoginConnectImpl::Params> Promise;
+		TFuture<FLoginConnectImpl::Params> Future = Promise.GetFuture();
 
-			// A success means the account is already linked, we'll process the NULL result in the next step
-		}
-		else if (Data->ResultCode == EOS_EResult::EOS_InvalidUser && Data->ContinuanceToken != nullptr)
+		TDefaultErrorResult<FGetExternalAuthTokenImpl> AuthTokenResult = GetExternalAuthTokenImpl(FGetExternalAuthTokenImpl::Params{AccountInfoEOS->EpicAccountId});
+		if (AuthTokenResult.IsError())
 		{
-			EOS_Auth_LinkAccountOptions LinkAccountOptions = {};
-			LinkAccountOptions.ApiVersion = EOS_AUTH_LINKACCOUNT_API_LATEST;
-			LinkAccountOptions.ContinuanceToken = Data->ContinuanceToken;
+			UE_LOG(LogOnlineServices, Warning, TEXT("[FAuthEOS::Login] Failure: GetExternalAuthTokenImpl %s"), *AuthTokenResult.GetErrorValue().GetLogString());
+			InAsyncOp.SetError(Errors::Unknown(MoveTemp(AuthTokenResult.GetErrorValue())));
 
-			EOS_Async(EOS_Auth_LinkAccount, AuthHandle, LinkAccountOptions, MoveTemp(Promise));
-			return;
-		}
-		else
-		{
-			FOnlineError Error = Errors::Unknown();
-			if (Data->ResultCode == EOS_EResult::EOS_InvalidAuth)
+			// Failed to acquire token - logout EAS.
+			LogoutEASImpl(FLogoutEASImpl::Params{ AccountInfoEOS->EpicAccountId })
+			.Next([Promise = MoveTemp(Promise)](TDefaultErrorResult<FLogoutEASImpl>&& LogoutResult) mutable -> void
 			{
-				Error = Errors::InvalidCreds();
-			}
+				if (LogoutResult.IsError())
+				{
+					UE_LOG(LogOnlineServices, Warning, TEXT("[FAuthEOS::Login] Failure: LogoutEASImpl %s"), *LogoutResult.GetErrorValue().GetLogString());
+				}
+				Promise.EmplaceValue(FLoginConnectImpl::Params{});
+			});
 
-			InAsyncOp.SetError(MoveTemp(Error));
+			return Future;
 		}
 
-		Promise.EmplaceValue();
+		Promise.EmplaceValue(FLoginConnectImpl::Params{Params.PlatformUserId, MoveTemp(AuthTokenResult.GetOkValue().Token)});
+		return Future;
 	})
-	.Then([this](TOnlineAsyncOp<FAuthLogin>& InAsyncOp, const EOS_Auth_LinkAccountCallbackInfo* Data) -> TFuture<TSharedPtr<FEOSConnectLoginCredentials>>
+	// Step 4: Attempt connect login. On connect login failure handle logout of EAS.
+	.Then([this](TOnlineAsyncOp<FAuthLogin>& InAsyncOp, FLoginConnectImpl::Params&& LoginConnectParams)
 	{
-		UE_LOG(LogTemp, Verbose, TEXT("[FAuthEOS::Login] EOS_Auth_LinkAccount Result: [%s]"), (Data == nullptr) ? TEXT("Null") : *LexToString(Data->ResultCode));
-		
-		// If Data is NULL, it means the account was already linked
-		if (Data == nullptr || Data->ResultCode == EOS_EResult::EOS_Success)
+		const TSharedRef<FAccountInfoEOS>& AccountInfoEOS = GetOpDataChecked<TSharedRef<FAccountInfoEOS>>(InAsyncOp, AccountInfoKeyName);
+
+		TPromise<void> Promise;
+		TFuture<void> Future = Promise.GetFuture();
+
+		// Attempt connect login.
+		LoginConnectImpl(LoginConnectParams)
+		.Next([this, AccountInfoEOS, WeakOp = InAsyncOp.AsWeak(), Promise = MoveTemp(Promise)](TDefaultErrorResult<FLoginConnectImpl>&& LoginResult) mutable -> void
 		{
-			EOS_EpicAccountId EpicAccountId = nullptr;
-			if (Data == nullptr)
+			if (TSharedPtr<TOnlineAsyncOp<FAuthLogin>> Op = WeakOp.Pin())
 			{
-				EpicAccountId = *InAsyncOp.Data.Get<EOS_EpicAccountId>(TEXT("EpicAccountId"));
+				if (LoginResult.IsError())
+				{
+					UE_LOG(LogOnlineServices, Warning, TEXT("[FAuthEOS::Login] Failure: LoginConnectImpl %s"), *LoginResult.GetErrorValue().GetLogString());
+					Op->SetError(Errors::Unknown(MoveTemp(LoginResult.GetErrorValue())));
+
+					LogoutEASImpl(FLogoutEASImpl::Params{ AccountInfoEOS->EpicAccountId })
+					.Next([Promise = MoveTemp(Promise)](TDefaultErrorResult<FLogoutEASImpl>&& LogoutResult) mutable -> void
+					{
+						if (LogoutResult.IsError())
+						{
+							UE_LOG(LogOnlineServices, Warning, TEXT("[FAuthEOS::Login] Failure: LogoutEASImpl %s"), *LogoutResult.GetErrorValue().GetLogString());
+						}
+						Promise.EmplaceValue();
+					});
+				}
+				else
+				{
+					// Successful login.
+					AccountInfoEOS->ProductUserId = LoginResult.GetOkValue().ProductUserId;
+					Promise.EmplaceValue();
+				}
 			}
 			else
 			{
-				// We cache the Epic Account Id to use it in later stages of the login process
-				InAsyncOp.Data.Set(TEXT("EpicAccountId"), Data->LocalUserId);
-				EpicAccountId = Data->LocalUserId;
+				Promise.EmplaceValue();
 			}
+		});
 
-			// On success, attempt Connect Login
-			EOS_Auth_Token* AuthToken = nullptr;
-			EOS_Auth_CopyUserAuthTokenOptions CopyOptions = { };
-			CopyOptions.ApiVersion = EOS_AUTH_COPYUSERAUTHTOKEN_API_LATEST;
+		return Future;
+	})
+	// Step 5: Fetch dependent data.
+	.Then([this](TOnlineAsyncOp<FAuthLogin>& InAsyncOp)
+	{
+		const TSharedRef<FAccountInfoEOS>& AccountInfoEOS = GetOpDataChecked<TSharedRef<FAccountInfoEOS>>(InAsyncOp, AccountInfoKeyName);
 
-			EOS_EResult CopyResult = EOS_Auth_CopyUserAuthToken(AuthHandle, &CopyOptions, EpicAccountId, &AuthToken);
+		// Get display name
+		EOS_UserInfo_CopyUserInfoOptions Options = { };
+		Options.ApiVersion = EOS_USERINFO_COPYUSERINFO_API_LATEST;
+		Options.LocalUserId = AccountInfoEOS->EpicAccountId;
+		Options.TargetUserId = AccountInfoEOS->EpicAccountId;
 
-			UE_LOG(LogTemp, Verbose, TEXT("[FAuthEOS::Login] EOS_Auth_CopyUserAuthToken Result: [%s]"), *LexToString(CopyResult));
+		EOS_UserInfo* UserInfo = nullptr;
 
-			if (CopyResult == EOS_EResult::EOS_Success)
-			{
-				TSharedPtr<FEOSConnectLoginCredentials> ConnectLoginCredentials = MakeShared<FEOSConnectLoginCredentialsEAS>(AuthToken);
-				return MakeFulfilledPromise<TSharedPtr<FEOSConnectLoginCredentials>>(MoveTemp(ConnectLoginCredentials)).GetFuture();
-			}
-			else
-			{
-				// TODO: EAS Logout
-				InAsyncOp.SetError(Errors::Unknown()); // TODO
-			}
-		}
-		else
+		EOS_EResult CopyUserInfoResult = EOS_UserInfo_CopyUserInfo(UserInfoHandle, &Options, &UserInfo);
+		if (CopyUserInfoResult == EOS_EResult::EOS_Success)
 		{
-			InAsyncOp.SetError(Errors::Unknown()); // TODO
+			AccountInfoEOS->Attributes.Emplace(AccountAttributeData::DisplayName, UTF8_TO_TCHAR(UserInfo->DisplayName));
+			EOS_UserInfo_Release(UserInfo);
 		}
-
-		return MakeFulfilledPromise<TSharedPtr<FEOSConnectLoginCredentials>>(nullptr).GetFuture();
-	});
-}
-
-void FAuthEOS::ProcessSuccessfulLogin(TOnlineAsyncOp<FAuthLogin>& InAsyncOp)
-{
-	const EOS_EpicAccountId* EpicAccountId = InAsyncOp.Data.Get<EOS_EpicAccountId>(TEXT("EpicAccountId"));
-	const EOS_ProductUserId ProductUserId = *InAsyncOp.Data.Get<EOS_ProductUserId>(TEXT("ProductUserId"));
-	const FOnlineAccountIdHandle LocalUserId = CreateAccountId(EpicAccountId ? *EpicAccountId : nullptr, ProductUserId);
-
-	UE_LOG(LogTemp, Verbose, TEXT("[FAuthEOS::Login] Successfully logged in as [%s]"), *ToLogString(LocalUserId));
-
-	TSharedRef<FAccountInfoEOS> AccountInfo = MakeShared<FAccountInfoEOS>();
-	AccountInfo->PlatformUserId = InAsyncOp.GetParams().PlatformUserId;
-	AccountInfo->UserId = LocalUserId;
-	AccountInfo->LoginStatus = ELoginStatus::LoggedIn;
-
-	if (EpicAccountId)
-	{
-		// Get display name
-		if (EOS_HUserInfo UserInfoHandle = EOS_Platform_GetUserInfoInterface(static_cast<FOnlineServicesEOS&>(GetServices()).GetEOSPlatformHandle()))
+		else
 		{
-			EOS_UserInfo_CopyUserInfoOptions Options = { };
-			Options.ApiVersion = EOS_USERINFO_COPYUSERINFO_API_LATEST;
-			Options.LocalUserId = *EpicAccountId;
-			Options.TargetUserId = *EpicAccountId;
+			FOnlineError CopyUserInfoError(FromEOSError(CopyUserInfoResult));
+			UE_LOG(LogOnlineServices, Warning, TEXT("[FAuthEOS::Login] Failure: EOS_UserInfo_CopyUserInfo %s"), *CopyUserInfoError.GetLogString());
+			InAsyncOp.SetError(Errors::Unknown(MoveTemp(CopyUserInfoError)));
 
-			EOS_UserInfo* UserInfo = nullptr;
+			TPromise<void> Promise;
+			TFuture<void> Future = Promise.GetFuture();
 
-			EOS_EResult CopyResult = EOS_UserInfo_CopyUserInfo(UserInfoHandle, &Options, &UserInfo);
-			if (CopyResult == EOS_EResult::EOS_Success)
+			LogoutEASImpl(FLogoutEASImpl::Params{ AccountInfoEOS->EpicAccountId })
+			.Next([Promise = MoveTemp(Promise)](TDefaultErrorResult<FLogoutEASImpl>&& LogoutResult) mutable -> void
 			{
-				AccountInfo->DisplayName = UTF8_TO_TCHAR(UserInfo->DisplayName);
-				EOS_UserInfo_Release(UserInfo);
-			}
+				if (LogoutResult.IsError())
+				{
+					UE_LOG(LogOnlineServices, Warning, TEXT("[FAuthEOS::Login] Failure: LogoutEASImpl %s"), *LogoutResult.GetErrorValue().GetLogString());
+				}
+				Promise.EmplaceValue();
+			});
+
+			return Future;
 		}
-	}
 
-	check(!AccountInfos.Contains(LocalUserId));
-	AccountInfos.Emplace(LocalUserId, AccountInfo);
+		return MakeFulfilledPromise<void>().GetFuture();
+	})
+	// Step 6: bookkeeping and notifications.
+	.Then([this](TOnlineAsyncOp<FAuthLogin>& InAsyncOp)
+	{
+		const TSharedRef<FAccountInfoEOS>& AccountInfoEOS = GetOpDataChecked<TSharedRef<FAccountInfoEOS>>(InAsyncOp, AccountInfoKeyName);
+		AccountInfoEOS->LoginStatus = ELoginStatus::LoggedIn;
+		AccountInfoEOS->AccountId = CreateAccountId(AccountInfoEOS->EpicAccountId, AccountInfoEOS->ProductUserId);
+		AccountInfoRegistryEOS.Register(AccountInfoEOS);
 
-	FAuthLogin::Result Result = { AccountInfo };
-	InAsyncOp.SetResult(MoveTemp(Result));
+		UE_LOG(LogOnlineServices, Log, TEXT("[FAuthEOS::Login] Successfully logged in as [%s]"), *ToLogString(AccountInfoEOS->AccountId));
+		OnAuthLoginStatusChangedEvent.Broadcast(FAuthLoginStatusChanged{ AccountInfoEOS, AccountInfoEOS->LoginStatus });
+		InAsyncOp.SetResult(FAuthLogin::Result{AccountInfoEOS});
+	})
+	.Enqueue(GetSerialQueue());
 
-	// When a user logs in, OnEASLoginStatusChanged can not trigger (if it's that user's first login) or trigger before we add relevant information to AccountInfos, so we trigger the status change event here 
-	FLoginStatusChanged EventParameters;
-	EventParameters.LocalUserId = LocalUserId;
-	EventParameters.PreviousStatus = ELoginStatus::NotLoggedIn;
-	EventParameters.CurrentStatus = ELoginStatus::LoggedIn;
-	OnLoginStatusChangedEvent.Broadcast(EventParameters);
+	return Op->GetHandle();
 }
 
-void FAuthEOS::OnEASLoginStatusChanged(FOnlineAccountIdHandle LocalUserId, ELoginStatus PreviousStatus, ELoginStatus CurrentStatus)
+template <typename IdType>
+TFuture<FOnlineAccountIdHandle> ResolveAccountIdImpl(FAuthEOS& AuthEOS, const FOnlineAccountIdHandle& LocalUserId, const IdType InId)
 {
-	UE_LOG(LogTemp, Warning, TEXT("OnEASLoginStatusChanged: [%s] [%s]->[%s]"), *ToLogString(LocalUserId), LexToString(PreviousStatus), LexToString(CurrentStatus));
-	if (TSharedRef<FAccountInfoEOS>* AccountInfoPtr = AccountInfos.Find(LocalUserId))
+	TPromise<FOnlineAccountIdHandle> Promise;
+	TFuture<FOnlineAccountIdHandle> Future = Promise.GetFuture();
+
+	AuthEOS.ResolveAccountIds(LocalUserId, { InId }).Next([Promise = MoveTemp(Promise)](const TArray<FOnlineAccountIdHandle>& AccountIds) mutable
 	{
-		FAccountInfoEOS& AccountInfo = **AccountInfoPtr;
-		if (AccountInfo.LoginStatus != CurrentStatus)
+		FOnlineAccountIdHandle Result;
+		if (AccountIds.Num() == 1)
 		{
-			FLoginStatusChanged EventParameters;
-			EventParameters.LocalUserId = LocalUserId;
-			EventParameters.PreviousStatus = AccountInfo.LoginStatus;
-			EventParameters.CurrentStatus = CurrentStatus;
-
-			AccountInfo.LoginStatus = CurrentStatus;
-
-			if (CurrentStatus == ELoginStatus::NotLoggedIn)
-			{
-				// Remove user
-				AccountInfos.Remove(LocalUserId); // Invalidates AccountInfo
-			}
-
-			OnLoginStatusChangedEvent.Broadcast(EventParameters);
+			Result = AccountIds[0];
 		}
-	}
+		Promise.SetValue(Result);
+	});
+
+	return Future;
 }
 
-TOnlineAsyncOpHandle<FAuthLogout> FAuthEOS::Logout(FAuthLogout::Params&& Params)
+TOnlineAsyncOpHandle<FAuthQueryExternalServerAuthTicket> FAuthEOS::QueryExternalServerAuthTicket(FAuthQueryExternalServerAuthTicket::Params&& Params)
 {
-	const FOnlineAccountIdHandle LocalUserId = Params.LocalUserId;
-	TOnlineAsyncOpRef<FAuthLogout> Op = GetOp<FAuthLogout>(MoveTemp(Params));
+	TOnlineAsyncOpRef<FAuthQueryExternalServerAuthTicket> Op = GetOp<FAuthQueryExternalServerAuthTicket>(MoveTemp(Params));
 
-	if (!ValidateOnlineId(LocalUserId))
-	{
-		Op->SetError(Errors::InvalidUser());
-		return Op->GetHandle();
-	}
-
-	const EOS_EpicAccountId LocalUserEasId = GetEpicAccountId(Params.LocalUserId);
-	if (!EOS_EpicAccountId_IsValid(LocalUserEasId) || !AccountInfos.Contains(LocalUserId))
-	{
-		// TODO: Error codes
-		Op->SetError(Errors::Unknown());
-		return Op->GetHandle();
-	}
-
-	// Should we destroy persistent auth first?
-	TOnlineChainableAsyncOp<FAuthLogout, void> NextOp = *Op;
-	if (Params.bDestroyAuth)
+	Op->Then([this](TOnlineAsyncOp<FAuthQueryExternalServerAuthTicket>& InAsyncOp)
 	{
-		EOS_Auth_DeletePersistentAuthOptions DeletePersistentAuthOptions = {0};
-		DeletePersistentAuthOptions.ApiVersion = EOS_AUTH_DELETEPERSISTENTAUTH_API_LATEST;
-		DeletePersistentAuthOptions.RefreshToken = nullptr; // Is this needed?  Docs say it's needed for consoles
-		NextOp = NextOp.Then([this, DeletePersistentAuthOptions](TOnlineAsyncOp<FAuthLogout>& InAsyncOp, TPromise<const EOS_Auth_DeletePersistentAuthCallbackInfo*>&& Promise)
-		{
-			EOS_Async(EOS_Auth_DeletePersistentAuth, AuthHandle, DeletePersistentAuthOptions, MoveTemp(Promise));
-		})
-		.Then([](TOnlineAsyncOp<FAuthLogout>& InAsyncOp, const EOS_Auth_DeletePersistentAuthCallbackInfo* Data)
+		const FAuthQueryExternalServerAuthTicket::Params& Params = InAsyncOp.GetParams();
+		TSharedPtr<FAccountInfoEOS> AccountInfoEOS = AccountInfoRegistryEOS.Find(Params.LocalUserId);
+		if (!AccountInfoEOS)
 		{
-			UE_LOG(LogTemp, Warning, TEXT("DeletePersistentAuthResult: [%s]"), ANSI_TO_TCHAR(EOS_EResult_ToString(Data->ResultCode)));
-			// Regardless of success/failure, continue
-		});
-	}
+			InAsyncOp.SetError(Errors::InvalidParams());
+			return;
+		}
 
-	// Logout
-	NextOp.Then([this, LocalUserEasId](TOnlineAsyncOp<FAuthLogout>& InAsyncOp, TPromise<const EOS_Auth_LogoutCallbackInfo*>&& Promise)
-	{
-		EOS_Auth_LogoutOptions LogoutOptions = { };
-		LogoutOptions.ApiVersion = EOS_AUTH_LOGOUT_API_LATEST;
-		LogoutOptions.LocalUserId = LocalUserEasId;
-		EOS_Async(EOS_Auth_Logout, AuthHandle, LogoutOptions, MoveTemp(Promise));
-	})
-	.Then([](TOnlineAsyncOp<FAuthLogout>& InAsyncOp, const EOS_Auth_LogoutCallbackInfo* Data)
-	{
-		UE_LOG(LogTemp, Warning, TEXT("LogoutResult: [%s]"), *LexToString(Data->ResultCode));
+		EOS_Auth_CopyUserAuthTokenOptions CopyUserAuthTokenOptions = {};
+		CopyUserAuthTokenOptions.ApiVersion = EOS_AUTH_COPYUSERAUTHTOKEN_API_LATEST;
 
-		if (Data->ResultCode == EOS_EResult::EOS_Success)
+		EOS_Auth_Token* AuthToken = nullptr;
+
+		EOS_EResult Result = EOS_Auth_CopyUserAuthToken(AuthHandle, &CopyUserAuthTokenOptions, AccountInfoEOS->EpicAccountId, &AuthToken);
+		if (Result == EOS_EResult::EOS_Success)
 		{
-			// Success
-			InAsyncOp.SetResult(FAuthLogout::Result());
+			ON_SCOPE_EXIT
+			{
+				EOS_Auth_Token_Release(AuthToken);
+			};
 
-			// OnLoginStatusChanged will be triggered by OnEASLoginStatusChanged
+			FExternalServerAuthTicket ExternalServerAuthTicket;
+			ExternalServerAuthTicket.Type = ExternalLoginType::Epic;
+			ExternalServerAuthTicket.Data = UTF8_TO_TCHAR(AuthToken->AccessToken);
+			InAsyncOp.SetResult(FAuthQueryExternalServerAuthTicket::Result{MoveTemp(ExternalServerAuthTicket)});
 		}
 		else
 		{
-			// TODO: Error codes
-			FOnlineError Error = Errors::Unknown();
-			InAsyncOp.SetError(MoveTemp(Error));
+			InAsyncOp.SetError(Errors::Unknown(FromEOSError(Result)));
 		}
-	}).Enqueue(GetSerialQueue());
-	
+	})
+	.Enqueue(GetSerialQueue());
+
 	return Op->GetHandle();
 }
 
-template <typename IdType>
-TFuture<FOnlineAccountIdHandle> ResolveAccountIdImpl(FAuthEOS& AuthEOS, const FOnlineAccountIdHandle& LocalUserId, const IdType InId)
+TOnlineAsyncOpHandle<FAuthQueryExternalAuthToken> FAuthEOS::QueryExternalAuthToken(FAuthQueryExternalAuthToken::Params&& Params)
 {
-	TPromise<FOnlineAccountIdHandle> Promise;
-	TFuture<FOnlineAccountIdHandle> Future = Promise.GetFuture();
+	TOnlineAsyncOpRef<FAuthQueryExternalAuthToken> Op = GetOp<FAuthQueryExternalAuthToken>(MoveTemp(Params));
 
-	AuthEOS.ResolveAccountIds(LocalUserId, { InId }).Next([Promise = MoveTemp(Promise)](const TArray<FOnlineAccountIdHandle>& AccountIds) mutable
+	Op->Then([this](TOnlineAsyncOp<FAuthQueryExternalAuthToken>& InAsyncOp)
 	{
-		FOnlineAccountIdHandle Result;
-		if (AccountIds.Num() == 1)
+		const FAuthQueryExternalAuthToken::Params& Params = InAsyncOp.GetParams();
+		TSharedPtr<FAccountInfoEOS> AccountInfoEOS = AccountInfoRegistryEOS.Find(Params.LocalUserId);
+		if (!AccountInfoEOS)
 		{
-			Result = AccountIds[0];
+			InAsyncOp.SetError(Errors::InvalidParams());
+			return;
 		}
-		Promise.SetValue(Result);
-	});
 
-	return Future;
+		TDefaultErrorResult<FGetExternalAuthTokenImpl> AuthTokenResult = GetExternalAuthTokenImpl(FGetExternalAuthTokenImpl::Params{ AccountInfoEOS->EpicAccountId });
+		if (AuthTokenResult.IsError())
+		{
+			UE_LOG(LogOnlineServices, Warning, TEXT("[FAuthEOS::QueryExternalAuthToken] Failure: GetExternalAuthTokenImpl %s"), *AuthTokenResult.GetErrorValue().GetLogString());
+			InAsyncOp.SetError(Errors::Unknown(MoveTemp(AuthTokenResult.GetErrorValue())));
+			return;
+		}
+
+		InAsyncOp.SetResult(FAuthQueryExternalAuthToken::Result{MoveTemp(AuthTokenResult.GetOkValue().Token)});
+	})
+	.Enqueue(GetSerialQueue());
+
+	return Op->GetHandle();
 }
 
 TFuture<FOnlineAccountIdHandle> FAuthEOS::ResolveAccountId(const FOnlineAccountIdHandle& LocalUserId, const EOS_EpicAccountId EpicAccountId)
@@ -552,6 +391,9 @@ TFuture<TArray<FOnlineAccountIdHandle>> FAuthEOS::ResolveAccountIds(const FOnlin
 		return MakeFulfilledPromise<TArray<FOnlineAccountIdHandle>>().GetFuture();
 	}
 
+	TPromise<TArray<FOnlineAccountIdHandle>> Promise;
+	TFuture<TArray<FOnlineAccountIdHandle>> Future = Promise.GetFuture();
+
 	TArray<FEpicAccountIdStrBuffer> EpicAccountIdStrsToQuery;
 	EpicAccountIdStrsToQuery.Reserve(MissingEpicAccountIds.Num());
 	for (const EOS_EpicAccountId EpicAccountId : MissingEpicAccountIds)
@@ -569,9 +411,6 @@ TFuture<TArray<FOnlineAccountIdHandle>> FAuthEOS::ResolveAccountIds(const FOnlin
 	TArray<const char*> EpicAccountIdStrPtrs;
 	Algo::Transform(EpicAccountIdStrsToQuery, EpicAccountIdStrPtrs, [](const FEpicAccountIdStrBuffer& Str) { return &Str[0]; });
 
-	TPromise<TArray<FOnlineAccountIdHandle>> Promise;
-	TFuture<TArray<FOnlineAccountIdHandle>> Future = Promise.GetFuture();
-
 	EOS_Connect_QueryExternalAccountMappingsOptions Options = {};
 	Options.ApiVersion = EOS_CONNECT_QUERYEXTERNALACCOUNTMAPPINGS_API_LATEST;
 	Options.LocalUserId = GetProductUserIdChecked(LocalUserId);
@@ -580,7 +419,7 @@ TFuture<TArray<FOnlineAccountIdHandle>> FAuthEOS::ResolveAccountIds(const FOnlin
 	Options.ExternalAccountIdCount = 1;
 
 	EOS_Async(EOS_Connect_QueryExternalAccountMappings, ConnectHandle, Options,
-	[this, WeakThis = AsWeak(), InEpicAccountIds, Promise = MoveTemp(Promise)](const EOS_Connect_QueryExternalAccountMappingsCallbackInfo* Data) mutable
+	[this, WeakThis = AsWeak(), InEpicAccountIds, Promise = MoveTemp(Promise)](const EOS_Connect_QueryExternalAccountMappingsCallbackInfo* Data) mutable -> void
 	{
 		TArray<FOnlineAccountIdHandle> AccountIds;
 		if (const TSharedPtr<IAuth> StrongThis = WeakThis.Pin())
@@ -610,7 +449,7 @@ TFuture<TArray<FOnlineAccountIdHandle>> FAuthEOS::ResolveAccountIds(const FOnlin
 			}
 			else
 			{
-				UE_LOG(LogTemp, Warning, TEXT("ResolveAccountId failed to query external mapping Result=[%s]"), *LexToString(Data->ResultCode));
+				UE_LOG(LogOnlineServices, Warning, TEXT("ResolveAccountId failed to query external mapping Result=[%s]"), *LexToString(Data->ResultCode));
 			}
 		}
 		Promise.SetValue(MoveTemp(AccountIds));
@@ -661,9 +500,8 @@ TFuture<TArray<FOnlineAccountIdHandle>> FAuthEOS::ResolveAccountIds(const FOnlin
 	Options.LocalUserId = GetProductUserIdChecked(LocalUserId);
 	Options.ProductUserIds = MissingProductUserIds.GetData();
 	Options.ProductUserIdCount = MissingProductUserIds.Num();
-
 	EOS_Async(EOS_Connect_QueryProductUserIdMappings, ConnectHandle, Options,
-	[this, WeakThis = AsWeak(), InProductUserIds, Promise = MoveTemp(Promise)](const EOS_Connect_QueryProductUserIdMappingsCallbackInfo* Data) mutable
+	[this, WeakThis = AsWeak(), InProductUserIds, Promise = MoveTemp(Promise)](const EOS_Connect_QueryProductUserIdMappingsCallbackInfo* Data) mutable -> void
 	{
 		TArray<FOnlineAccountIdHandle> AccountIds;
 		if (const TSharedPtr<IAuth> StrongThis = WeakThis.Pin())
@@ -697,7 +535,7 @@ TFuture<TArray<FOnlineAccountIdHandle>> FAuthEOS::ResolveAccountIds(const FOnlin
 			}
 			else
 			{
-				UE_LOG(LogTemp, Warning, TEXT("ResolveAccountId failed to query external mapping Result=[%s]"), *LexToString(Data->ResultCode));
+				UE_LOG(LogOnlineServices, Warning, TEXT("ResolveAccountId failed to query external mapping Result=[%s]"), *LexToString(Data->ResultCode));
 			}
 		}
 		Promise.SetValue(MoveTemp(AccountIds));
diff --git a/Engine/Plugins/Online/Experimental/OnlineServicesEOS/Source/Public/Online/AuthEOS.h b/Engine/Plugins/Online/Experimental/OnlineServicesEOS/Source/Public/Online/AuthEOS.h
--- a/Engine/Plugins/Online/Experimental/OnlineServicesEOS/Source/Public/Online/AuthEOS.h
+++ b/Engine/Plugins/Online/Experimental/OnlineServicesEOS/Source/Public/Online/AuthEOS.h
@@ -4,6 +4,8 @@
 
 #include "Online/AuthEOSGS.h"
 
+#include "eos_userinfo_types.h"
+
 namespace UE::Online {
 
 class FOnlineServicesEOS;
@@ -22,7 +24,8 @@ public:
 
 	// Begin IAuth
 	virtual TOnlineAsyncOpHandle<FAuthLogin> Login(FAuthLogin::Params&& Params) override;
-	virtual TOnlineAsyncOpHandle<FAuthLogout> Logout(FAuthLogout::Params&& Params) override;
+	virtual TOnlineAsyncOpHandle<FAuthQueryExternalServerAuthTicket> QueryExternalServerAuthTicket(FAuthQueryExternalServerAuthTicket::Params&& Params) override;
+	virtual TOnlineAsyncOpHandle<FAuthQueryExternalAuthToken> QueryExternalAuthToken(FAuthQueryExternalAuthToken::Params&& Params) override;
 	// End IAuth
 
 	// Begin FAuthEOSGS
@@ -38,13 +41,12 @@ public:
 	TFunction<TFuture<TArray<FOnlineAccountIdHandle>>(FOnlineAsyncOp& InAsyncOp, const TArray<EOS_EpicAccountId>& EpicAccountIds)> ResolveEpicIdsFn();
 
 protected:
-	TOnlineChainableAsyncOp<FAuthLogin, TSharedPtr<FEOSConnectLoginCredentials>> LoginEAS(TOnlineAsyncOp<FAuthLogin>& InAsyncOp);
 	void ProcessSuccessfulLogin(TOnlineAsyncOp<FAuthLogin>& InAsyncOp);
 	void OnEASLoginStatusChanged(FOnlineAccountIdHandle LocalUserId, ELoginStatus PreviousStatus, ELoginStatus CurrentStatus);
 
 	static FOnlineAccountIdHandle CreateAccountId(const EOS_EpicAccountId EpicAccountId, const EOS_ProductUserId ProductUserId);
 
-	EOS_NotificationId NotifyEASLoginStatusChangedNotificationId = EOS_INVALID_NOTIFICATIONID;
+	EOS_HUserInfo UserInfoHandle = nullptr;
 };
 
 /* UE::Online */ }
diff --git a/Engine/Plugins/Online/Experimental/OnlineServicesEOSGS/Source/Private/Online/AuthEOSGS.cpp b/Engine/Plugins/Online/Experimental/OnlineServicesEOSGS/Source/Private/Online/AuthEOSGS.cpp
--- a/Engine/Plugins/Online/Experimental/OnlineServicesEOSGS/Source/Private/Online/AuthEOSGS.cpp
+++ b/Engine/Plugins/Online/Experimental/OnlineServicesEOSGS/Source/Private/Online/AuthEOSGS.cpp
@@ -2,7 +2,8 @@
 
 #include "Online/AuthEOSGS.h"
 
-#include "Algo/Transform.h"
+#include "Containers/BackgroundableTicker.h"
+#include "Math/UnrealMathUtility.h"
 #include "Misc/CommandLine.h"
 #include "Online/AuthErrors.h"
 #include "Online/OnlineAsyncOp.h"
@@ -22,154 +23,532 @@
 
 namespace UE::Online {
 
-struct FAuthEOSConfig
+struct FAuthEOSGSLoginConfig
 {
-	FString DefaultExternalCredentialTypeStr;
+	bool EASAuthEnabled = false;
+	TArray<FString> DefaultScopes;
+};
+
+struct FAuthEOSGSLoginRecoveryConfig
+{
+	float Interval = 30.f;
+	float Jitter = 5.f;
 };
 
 namespace Meta {
 
-BEGIN_ONLINE_STRUCT_META(FAuthEOSConfig)
-	ONLINE_STRUCT_FIELD(FAuthEOSConfig, DefaultExternalCredentialTypeStr)
+BEGIN_ONLINE_STRUCT_META(FAuthEOSGSLoginConfig)
+	ONLINE_STRUCT_FIELD(FAuthEOSGSLoginConfig, EASAuthEnabled),
+	ONLINE_STRUCT_FIELD(FAuthEOSGSLoginConfig, DefaultScopes)
+END_ONLINE_STRUCT_META()
+
+BEGIN_ONLINE_STRUCT_META(FAuthEOSGSLoginRecoveryConfig)
+	ONLINE_STRUCT_FIELD(FAuthEOSGSLoginRecoveryConfig, Interval),
+	ONLINE_STRUCT_FIELD(FAuthEOSGSLoginRecoveryConfig, Jitter)
 END_ONLINE_STRUCT_META()
 
 /* Meta*/ }
 
-ELoginStatus ToELoginStatus(EOS_ELoginStatus InStatus)
+namespace {
+
+static const FString AccountInfoKeyName = TEXT("AccountInfoEOS");
+
+/* anonymous */ }
+
+enum class EEOSConnectTranslationFlags : uint8
 {
-	switch (InStatus)
-	{
-	case EOS_ELoginStatus::EOS_LS_NotLoggedIn:
+	None = 0,
+	DisplayName = 1 << 0,
+};
+ENUM_CLASS_FLAGS(EEOSConnectTranslationFlags);
+
+struct FEOSConnectTranslationTraits
+{
+	FEOSConnectTranslationTraits() = delete;
+	EOS_EExternalCredentialType Type;
+	EEOSConnectTranslationFlags Flags;
+};
+
+class FEOSConnectLoginOptions : public EOS_Connect_LoginOptions
+{
+public:
+	FEOSConnectLoginOptions(const FEOSConnectLoginOptions&) = delete;
+	FEOSConnectLoginOptions& operator=(const FEOSConnectLoginOptions&) = delete;
+	FEOSConnectLoginOptions(FEOSConnectLoginOptions&&);
+	FEOSConnectLoginOptions& operator=(FEOSConnectLoginOptions&&);
+	virtual ~FEOSConnectLoginOptions() = default;
+
+	static TDefaultErrorResultInternal<FEOSConnectLoginOptions> Create(FPlatformUserId PlatformUserId, const FExternalAuthToken& ExternalAuthToken);
+private:
+	FEOSConnectLoginOptions();
+
+	static const FEOSConnectTranslationTraits* GetConnectTranslationTraits(FName ExternalLoginType);
+
+	EOS_Connect_Credentials CredentialsData;
+	EOS_Connect_UserLoginInfo UserLoginInfoData;
+	char DisplayNameUtf8[EOS_CONNECT_USERLOGININFO_DISPLAYNAME_MAX_LENGTH] = {};
+	TArray<char> ExternalAuthTokenUtf8;
+};
+
+const FEOSConnectTranslationTraits* FEOSConnectLoginOptions::GetConnectTranslationTraits(FName ExternalLoginType)
+{
+	static const TMap<FName, FEOSConnectTranslationTraits> SupportedLoginTranslatorTraits = {
+		{ ExternalLoginType::Epic, { EOS_EExternalCredentialType::EOS_ECT_EPIC, EEOSConnectTranslationFlags::None } },
+		{ ExternalLoginType::SteamAppTicket, { EOS_EExternalCredentialType::EOS_ECT_STEAM_APP_TICKET, EEOSConnectTranslationFlags::None } },
+		{ ExternalLoginType::PsnIdToken, { EOS_EExternalCredentialType::EOS_ECT_PSN_ID_TOKEN, EEOSConnectTranslationFlags::None } },
+		{ ExternalLoginType::XblXstsToken, { EOS_EExternalCredentialType::EOS_ECT_XBL_XSTS_TOKEN, EEOSConnectTranslationFlags::None } },
+		{ ExternalLoginType::DiscordAccessToken, { EOS_EExternalCredentialType::EOS_ECT_DISCORD_ACCESS_TOKEN, EEOSConnectTranslationFlags::None } },
+		{ ExternalLoginType::GogSessionTicket, { EOS_EExternalCredentialType::EOS_ECT_GOG_SESSION_TICKET, EEOSConnectTranslationFlags::None } },
+		{ ExternalLoginType::NintendoIdToken, { EOS_EExternalCredentialType::EOS_ECT_NINTENDO_ID_TOKEN, EEOSConnectTranslationFlags::DisplayName } },
+		{ ExternalLoginType::NintendoNsaIdToken, { EOS_EExternalCredentialType::EOS_ECT_NINTENDO_NSA_ID_TOKEN, EEOSConnectTranslationFlags::DisplayName } },
+		{ ExternalLoginType::UplayAccessToken, { EOS_EExternalCredentialType::EOS_ECT_UPLAY_ACCESS_TOKEN, EEOSConnectTranslationFlags::None } },
+		{ ExternalLoginType::OpenIdAccessToken, { EOS_EExternalCredentialType::EOS_ECT_OPENID_ACCESS_TOKEN, EEOSConnectTranslationFlags::None } },
+		{ ExternalLoginType::DeviceIdAccessToken, { EOS_EExternalCredentialType::EOS_ECT_DEVICEID_ACCESS_TOKEN, EEOSConnectTranslationFlags::DisplayName } },
+		{ ExternalLoginType::AppleIdToken, { EOS_EExternalCredentialType::EOS_ECT_APPLE_ID_TOKEN, EEOSConnectTranslationFlags::DisplayName } },
+		{ ExternalLoginType::GoogleIdToken, { EOS_EExternalCredentialType::EOS_ECT_GOOGLE_ID_TOKEN, EEOSConnectTranslationFlags::DisplayName } },
+		{ ExternalLoginType::OculusUserIdNonce, { EOS_EExternalCredentialType::EOS_ECT_OCULUS_USERID_NONCE, EEOSConnectTranslationFlags::DisplayName } },
+		{ ExternalLoginType::ItchioJwt, { EOS_EExternalCredentialType::EOS_ECT_ITCHIO_JWT, EEOSConnectTranslationFlags::None } },
+		{ ExternalLoginType::ItchioKey, { EOS_EExternalCredentialType::EOS_ECT_ITCHIO_KEY, EEOSConnectTranslationFlags::None } },
+		{ ExternalLoginType::EpicIdToken, { EOS_EExternalCredentialType::EOS_ECT_EPIC_ID_TOKEN, EEOSConnectTranslationFlags::None } },
+		{ ExternalLoginType::AmazonAccessToken, { EOS_EExternalCredentialType::EOS_ECT_AMAZON_ACCESS_TOKEN, EEOSConnectTranslationFlags::DisplayName } },
+	};
+
+	return SupportedLoginTranslatorTraits.Find(ExternalLoginType);
+}
+
+FEOSConnectLoginOptions::FEOSConnectLoginOptions(FEOSConnectLoginOptions&& Other)
+{
+	*this = MoveTemp(Other);
+}
+
+FEOSConnectLoginOptions& FEOSConnectLoginOptions::operator=(FEOSConnectLoginOptions&& Other)
+{
+	ApiVersion = Other.ApiVersion;
+	Credentials = nullptr;
+	UserLoginInfo = nullptr;
+
+	if (Other.Credentials)
 	{
-		return ELoginStatus::NotLoggedIn;
+		ExternalAuthTokenUtf8 = MoveTemp(Other.ExternalAuthTokenUtf8);
+		CredentialsData = Other.CredentialsData;
+		CredentialsData.Token = ExternalAuthTokenUtf8.GetData();
+		Credentials = &CredentialsData;
+		Other.CredentialsData = EOS_Connect_Credentials({});
+		memset(&Other.CredentialsData, 0, sizeof(Other.CredentialsData));
 	}
-	case EOS_ELoginStatus::EOS_LS_UsingLocalProfile:
+
+	if (Other.UserLoginInfo)
 	{
-		return ELoginStatus::UsingLocalProfile;
+		memcpy(DisplayNameUtf8, Other.DisplayNameUtf8, EOS_CONNECT_USERLOGININFO_DISPLAYNAME_MAX_LENGTH);
+		UserLoginInfoData = Other.UserLoginInfoData;
+		UserLoginInfoData.DisplayName = DisplayNameUtf8;
+		UserLoginInfo = &UserLoginInfoData;
+		memset(&Other.UserLoginInfo, 0, sizeof(Other.UserLoginInfo));
 	}
-	case EOS_ELoginStatus::EOS_LS_LoggedIn:
+
+	return *this;
+}
+
+FEOSConnectLoginOptions::FEOSConnectLoginOptions()
+{
+	// EOS_Connect_LoginOptions init
+	ApiVersion = EOS_CONNECT_LOGIN_API_LATEST;
+	Credentials = &CredentialsData;
+	UserLoginInfo = nullptr;
+
+	// EOS_Connect_Credentials init
+	CredentialsData.ApiVersion = EOS_CONNECT_CREDENTIALS_API_LATEST;
+	CredentialsData.Token = nullptr;
+
+	// EOS_Connect_UserLoginInfo init
+	UserLoginInfoData.ApiVersion = EOS_CONNECT_USERLOGININFO_API_LATEST;
+	UserLoginInfoData.DisplayName = DisplayNameUtf8;
+}
+
+TDefaultErrorResultInternal<FEOSConnectLoginOptions> FEOSConnectLoginOptions::Create(FPlatformUserId PlatformUserId, const FExternalAuthToken& ExternalAuthToken)
+{
+	const FEOSConnectTranslationTraits* TranslatorTraits = GetConnectTranslationTraits(ExternalAuthToken.Type);
+	if (TranslatorTraits == nullptr)
 	{
-		return ELoginStatus::LoggedIn;
+		UE_LOG(LogOnlineServices, Warning, TEXT("FEOSConnectLoginOptions::Create: Failed - Unsupported external auth type: %s"), *ExternalAuthToken.Type.ToString());
+		return TDefaultErrorResultInternal<FEOSConnectLoginOptions>(Errors::InvalidParams());
 	}
+
+	FEOSConnectLoginOptions EOSConnectLoginOptions;
+	if (EnumHasAnyFlags(TranslatorTraits->Flags, EEOSConnectTranslationFlags::DisplayName))
+	{
+		// Todo: Lookup platform OSS Auth interface to resolve display name using PlatformUserId.
+		FCStringAnsi::Strncpy(EOSConnectLoginOptions.DisplayNameUtf8, "PlatformUser", EOS_CONNECT_USERLOGININFO_DISPLAYNAME_MAX_LENGTH);
+		EOSConnectLoginOptions.UserLoginInfo = &EOSConnectLoginOptions.UserLoginInfoData;
 	}
-	return ELoginStatus::NotLoggedIn;
+
+	EOSConnectLoginOptions.ExternalAuthTokenUtf8.SetNumZeroed(ExternalAuthToken.Data.Len() + 1);
+	FCStringAnsi::Strncpy(EOSConnectLoginOptions.ExternalAuthTokenUtf8.GetData(), TCHAR_TO_UTF8(*ExternalAuthToken.Data), EOSConnectLoginOptions.ExternalAuthTokenUtf8.Num());
+	EOSConnectLoginOptions.CredentialsData.Token = EOSConnectLoginOptions.ExternalAuthTokenUtf8.GetData();
+	EOSConnectLoginOptions.CredentialsData.Type = TranslatorTraits->Type;
+	return TDefaultErrorResultInternal<FEOSConnectLoginOptions>(MoveTemp(EOSConnectLoginOptions));
 }
 
-FEOSAuthCredentials::FEOSAuthCredentials() :
-	EOS_Auth_Credentials()
+enum class EEOSAuthTranslationFlags : uint8
 {
-	ApiVersion = EOS_AUTH_CREDENTIALS_API_LATEST;
-	Id = IdAnsi;
-	Token = TokenAnsi;
-}
+	None = 0,
+	SetId = 1 << 0,
+	SetTokenFromString = 1 << 1,
+	SetTokenFromExternalAuth = 1 << 2,
+};
+ENUM_CLASS_FLAGS(EEOSAuthTranslationFlags);
 
+struct FEOSAuthTranslationTraits
+{
+	FEOSAuthTranslationTraits() = delete;
+	EOS_ELoginCredentialType Type;
+	EEOSAuthTranslationFlags Flags;
+};
 
-FEOSAuthCredentials::FEOSAuthCredentials(const FEOSAuthCredentials& Other)
+struct FEOSExternalAuthTranslationTraits
 {
-	ApiVersion = Other.ApiVersion;
-	Id = IdAnsi;
-	Token = TokenAnsi;
-	Type = Other.Type;
-	SystemAuthCredentialsOptions = Other.SystemAuthCredentialsOptions;
-	ExternalType = Other.ExternalType;
+	FEOSExternalAuthTranslationTraits() = delete;
+	EOS_EExternalCredentialType Type;
+};
 
-	FCStringAnsi::Strncpy(IdAnsi, Other.IdAnsi, EOS_STRING_BUFFER_LENGTH);
-	FCStringAnsi::Strncpy(TokenAnsi, Other.TokenAnsi, EOS_MAX_TOKEN_SIZE);
-}
+class FEOSAuthLoginOptions : public EOS_Auth_LoginOptions
+{
+public:
+	FEOSAuthLoginOptions(const FEOSAuthLoginOptions&) = delete;
+	FEOSAuthLoginOptions& operator=(const FEOSAuthLoginOptions&) = delete;
+	FEOSAuthLoginOptions(FEOSAuthLoginOptions&&);
+	FEOSAuthLoginOptions& operator=(FEOSAuthLoginOptions&&);
+
+	static TDefaultErrorResultInternal<FEOSAuthLoginOptions> Create(
+		FName CredentialsType,
+		const FString& CredentialsId,
+		const TVariant<FString, FExternalAuthToken>& CredentialsToken,
+		const TArray<FString>& Scopes);
+
+private:
+	static TDefaultErrorResultInternal<FEOSAuthLoginOptions> CreateImpl(
+		FName CredentialsType,
+		const FString& CredentialsId,
+		const TVariant<FString, FExternalAuthToken>& CredentialsToken,
+		const TArray<FString>& Scopes);
+
+	FEOSAuthLoginOptions();
+
+	static const FEOSAuthTranslationTraits* GetLoginTranslatorTraits(FName Name);
+	static const FEOSExternalAuthTranslationTraits* GetExternalAuthTranslationTraits(FName ExternalAuthType);
+	static const EOS_EAuthScopeFlags* GetScopeFlag(const FString& ScopeName);
+
+	EOS_Auth_Credentials CredentialsData;
+	TArray<char> IdUtf8;
+	TArray<char> TokenUtf8;
+};
 
-FEOSAuthCredentials::FEOSAuthCredentials(EOS_EExternalCredentialType InExternalType, const TArray<uint8>& InToken) :
-	EOS_Auth_Credentials()
+const FEOSAuthTranslationTraits* FEOSAuthLoginOptions::GetLoginTranslatorTraits(FName Name)
 {
-	ApiVersion = EOS_AUTH_CREDENTIALS_API_LATEST;
-	Id = IdAnsi;
-	Token = TokenAnsi;
+	static const TMap<FName, FEOSAuthTranslationTraits> SupportedLoginTranslatorTraits = {
+		{ LoginCredentialsType::Password, { EOS_ELoginCredentialType::EOS_LCT_Password, EEOSAuthTranslationFlags::SetId | EEOSAuthTranslationFlags::SetTokenFromString } },
+		{ LoginCredentialsType::ExchangeCode, { EOS_ELoginCredentialType::EOS_LCT_ExchangeCode, EEOSAuthTranslationFlags::SetTokenFromString } },
+		{ LoginCredentialsType::PersistentAuth, { EOS_ELoginCredentialType::EOS_LCT_PersistentAuth, EEOSAuthTranslationFlags::SetTokenFromString } },
+		{ LoginCredentialsType::Developer, { EOS_ELoginCredentialType::EOS_LCT_Developer, EEOSAuthTranslationFlags::SetId | EEOSAuthTranslationFlags::SetTokenFromString } },
+		{ LoginCredentialsType::RefreshToken, { EOS_ELoginCredentialType::EOS_LCT_RefreshToken, EEOSAuthTranslationFlags::SetTokenFromString } },
+		{ LoginCredentialsType::AccountPortal, { EOS_ELoginCredentialType::EOS_LCT_AccountPortal, EEOSAuthTranslationFlags::SetId | EEOSAuthTranslationFlags::SetTokenFromString } },
+		{ LoginCredentialsType::ExternalAuth, { EOS_ELoginCredentialType::EOS_LCT_ExternalAuth, EEOSAuthTranslationFlags::SetTokenFromExternalAuth } },
+	};
 
-	Type = EOS_ELoginCredentialType::EOS_LCT_ExternalAuth;
-	ExternalType = InExternalType;
+	return SupportedLoginTranslatorTraits.Find(Name);
+}
+
+const FEOSExternalAuthTranslationTraits* FEOSAuthLoginOptions::GetExternalAuthTranslationTraits(FName ExternalAuthType)
+{
+	static const TMap<FName, FEOSExternalAuthTranslationTraits> SupportedExternalAuthTraits = {
+		{ ExternalLoginType::Epic, { EOS_EExternalCredentialType::EOS_ECT_EPIC } },
+		{ ExternalLoginType::SteamAppTicket, { EOS_EExternalCredentialType::EOS_ECT_STEAM_APP_TICKET } },
+		{ ExternalLoginType::PsnIdToken, { EOS_EExternalCredentialType::EOS_ECT_PSN_ID_TOKEN } },
+		{ ExternalLoginType::XblXstsToken, { EOS_EExternalCredentialType::EOS_ECT_XBL_XSTS_TOKEN } },
+		{ ExternalLoginType::DiscordAccessToken, { EOS_EExternalCredentialType::EOS_ECT_DISCORD_ACCESS_TOKEN } },
+		{ ExternalLoginType::GogSessionTicket, { EOS_EExternalCredentialType::EOS_ECT_GOG_SESSION_TICKET } },
+		{ ExternalLoginType::NintendoIdToken, { EOS_EExternalCredentialType::EOS_ECT_NINTENDO_ID_TOKEN } },
+		{ ExternalLoginType::NintendoNsaIdToken, { EOS_EExternalCredentialType::EOS_ECT_NINTENDO_NSA_ID_TOKEN } },
+		{ ExternalLoginType::UplayAccessToken, { EOS_EExternalCredentialType::EOS_ECT_UPLAY_ACCESS_TOKEN } },
+		{ ExternalLoginType::OpenIdAccessToken, { EOS_EExternalCredentialType::EOS_ECT_OPENID_ACCESS_TOKEN } },
+		{ ExternalLoginType::DeviceIdAccessToken, { EOS_EExternalCredentialType::EOS_ECT_DEVICEID_ACCESS_TOKEN } },
+		{ ExternalLoginType::AppleIdToken, { EOS_EExternalCredentialType::EOS_ECT_APPLE_ID_TOKEN } },
+		{ ExternalLoginType::GoogleIdToken, { EOS_EExternalCredentialType::EOS_ECT_GOOGLE_ID_TOKEN } },
+		{ ExternalLoginType::OculusUserIdNonce, { EOS_EExternalCredentialType::EOS_ECT_OCULUS_USERID_NONCE } },
+		{ ExternalLoginType::ItchioJwt, { EOS_EExternalCredentialType::EOS_ECT_ITCHIO_JWT } },
+		{ ExternalLoginType::ItchioKey, { EOS_EExternalCredentialType::EOS_ECT_ITCHIO_KEY } },
+		{ ExternalLoginType::EpicIdToken, { EOS_EExternalCredentialType::EOS_ECT_EPIC_ID_TOKEN } },
+		{ ExternalLoginType::AmazonAccessToken, { EOS_EExternalCredentialType::EOS_ECT_AMAZON_ACCESS_TOKEN } },
+	};
 
-	uint32_t InOutBufferLength = EOS_MAX_TOKEN_SIZE;
-	EOS_ByteArray_ToString(InToken.GetData(), InToken.Num(), TokenAnsi, &InOutBufferLength);
+	return SupportedExternalAuthTraits.Find(ExternalAuthType);
 }
 
-FEOSAuthCredentials& FEOSAuthCredentials::operator=(FEOSAuthCredentials& Other)
+const EOS_EAuthScopeFlags* FEOSAuthLoginOptions::GetScopeFlag(const FString& ScopeName)
 {
-	ApiVersion = Other.ApiVersion;
-	Type = Other.Type;
-	SystemAuthCredentialsOptions = Other.SystemAuthCredentialsOptions;
-	ExternalType = Other.ExternalType;
+	static const TMap<FString, EOS_EAuthScopeFlags> SupportedScopes = {
+		{ TEXT("BasicProfile"), EOS_EAuthScopeFlags::EOS_AS_BasicProfile },
+		{ TEXT("FriendsList"), EOS_EAuthScopeFlags::EOS_AS_FriendsList },
+		{ TEXT("Presence"), EOS_EAuthScopeFlags::EOS_AS_Presence },
+		{ TEXT("FriendsManagement"), EOS_EAuthScopeFlags::EOS_AS_FriendsManagement },
+		{ TEXT("Email"), EOS_EAuthScopeFlags::EOS_AS_Email },
+	};
 
-	FCStringAnsi::Strncpy(IdAnsi, Other.IdAnsi, EOS_STRING_BUFFER_LENGTH);
-	FCStringAnsi::Strncpy(TokenAnsi, Other.TokenAnsi, EOS_MAX_TOKEN_SIZE);
+	return SupportedScopes.Find(ScopeName);
+}
 
-	return *this;
+FEOSAuthLoginOptions::FEOSAuthLoginOptions(FEOSAuthLoginOptions&& Other)
+{
+	*this = MoveTemp(Other);
 }
 
-void FEOSAuthCredentials::SetToken(const FCredentialsToken& InToken)
+FEOSAuthLoginOptions& FEOSAuthLoginOptions::operator=(FEOSAuthLoginOptions&& Other)
 {
-	if (InToken.IsType<TArray<uint8>>())
+	CredentialsData = Other.CredentialsData;
+
+	// Pointer fixup.
+	if (CredentialsData.Id)
 	{
-		const TArray<uint8>& TokenData = InToken.Get<TArray<uint8>>();
-		uint32_t InOutBufferLength = EOS_MAX_TOKEN_SIZE;
-		EOS_ByteArray_ToString(TokenData.GetData(), TokenData.Num(), TokenAnsi, &InOutBufferLength);
+		IdUtf8 = MoveTemp(Other.IdUtf8);
+		CredentialsData.Id = IdUtf8.GetData();
 	}
-	else if (InToken.IsType<FString>())
+	if (CredentialsData.Token)
 	{
-		const FString& TokenString = InToken.Get<FString>();
-		FCStringAnsi::Strncpy(TokenAnsi, TCHAR_TO_UTF8(*TokenString), EOS_MAX_TOKEN_SIZE);
+		TokenUtf8 = MoveTemp(Other.TokenUtf8);
+		CredentialsData.Token = TokenUtf8.GetData();
 	}
-	else
+	if (CredentialsData.SystemAuthCredentialsOptions)
 	{
-		UE_LOG(LogTemp, Warning, TEXT("SetToken cannot be set with an invalid credentials token parameter. Please ensure there is valid data in the credentials token."));
+		// todo
 	}
+
+	Credentials = &CredentialsData;
+	ApiVersion = Other.ApiVersion;
+	ScopeFlags = Other.ScopeFlags;
+
+	Other.Credentials = nullptr;
+	return *this;
 }
 
-FEOSConnectLoginCredentials::FEOSConnectLoginCredentials()
+FEOSAuthLoginOptions::FEOSAuthLoginOptions()
 {
-	ApiVersion = EOS_CONNECT_CREDENTIALS_API_LATEST;
+	// EOS_Auth_LoginOptions init
+	ApiVersion = EOS_AUTH_LOGIN_API_LATEST;
+	Credentials = &CredentialsData;
+	ScopeFlags = EOS_EAuthScopeFlags::EOS_AS_NoFlags;
+
+	// EOS_Auth_Credentials init
+	CredentialsData.ApiVersion = EOS_AUTH_CREDENTIALS_API_LATEST;
+	CredentialsData.Id = nullptr;
+	CredentialsData.Token = nullptr;
+	CredentialsData.Type = EOS_ELoginCredentialType::EOS_LCT_Password;
+	CredentialsData.SystemAuthCredentialsOptions = nullptr;
+	CredentialsData.ExternalType = EOS_EExternalCredentialType::EOS_ECT_EPIC;
 }
 
-class FEOSConnectLoginCredentialsEOS : public FEOSConnectLoginCredentials
+TDefaultErrorResultInternal<FEOSAuthLoginOptions> FEOSAuthLoginOptions::Create(
+	FName CredentialsType,
+	const FString& CredentialsId,
+	const TVariant<FString, FExternalAuthToken>& CredentialsToken,
+	const TArray<FString>& Scopes)
 {
-public:
-	FEOSConnectLoginCredentialsEOS()
+	if (CredentialsType == LoginCredentialsType::Auto)
+	{
+		FString	CommandLineAuthType;
+		FString	CommandLineAuthId;
+		FString	CommandLineAuthToken;
+		FParse::Value(FCommandLine::Get(), TEXT("AUTH_TYPE="), CommandLineAuthType);
+		FParse::Value(FCommandLine::Get(), TEXT("AUTH_LOGIN="), CommandLineAuthId);
+		FParse::Value(FCommandLine::Get(), TEXT("AUTH_PASSWORD="), CommandLineAuthToken);
+		if (!CommandLineAuthId.IsEmpty() && !CommandLineAuthToken.IsEmpty() && !CommandLineAuthType.IsEmpty())
+		{
+			return CreateImpl(
+				*CommandLineAuthType,
+				CommandLineAuthId,
+				TVariant<FString, FExternalAuthToken>(TInPlaceType<FString>(),
+				MoveTemp(CommandLineAuthToken)),
+				Scopes);
+		}
+		else
+		{
+			return TDefaultErrorResultInternal<FEOSAuthLoginOptions>(Errors::InvalidCreds());
+		}
+	}
+	else
+	{
+		return CreateImpl(CredentialsType, CredentialsId, CredentialsToken, Scopes);
+	}
+}
+
+TDefaultErrorResultInternal<FEOSAuthLoginOptions> FEOSAuthLoginOptions::CreateImpl(
+	FName CredentialsType,
+	const FString& CredentialsId,
+	const TVariant<FString, FExternalAuthToken>& CredentialsToken,
+	const TArray<FString>& Scopes)
+{
+	const FEOSAuthTranslationTraits* TranslatorTraits = GetLoginTranslatorTraits(CredentialsType);
+	if (TranslatorTraits == nullptr)
+	{
+		UE_LOG(LogOnlineServices, Warning, TEXT("FEOSAuthLoginOptions::Create: Failed - Unsupported login type: %s"), *CredentialsType.ToString());
+		return TDefaultErrorResultInternal<FEOSAuthLoginOptions>(Errors::InvalidParams());
+	}
+
+	FEOSAuthLoginOptions EOSAuthLoginOptions;
+	EOSAuthLoginOptions.CredentialsData.Type = TranslatorTraits->Type;
+
+	if (EnumHasAnyFlags(TranslatorTraits->Flags, EEOSAuthTranslationFlags::SetId))
 	{
-		Token = TokenAnsi;
-		TokenAnsi[0] = '\0';
+		FTCHARToUTF8 Converter(*CredentialsId);
+		EOSAuthLoginOptions.IdUtf8.SetNumZeroed(Converter.Length() + 1);
+		FCStringAnsi::Strncpy(EOSAuthLoginOptions.IdUtf8.GetData(), Converter.Get(), EOSAuthLoginOptions.IdUtf8.Num());
+		EOSAuthLoginOptions.CredentialsData.Id = EOSAuthLoginOptions.IdUtf8.GetData();
 	}
-	void SetToken(const FCredentialsToken& InToken)
+
+	if (EnumHasAnyFlags(TranslatorTraits->Flags, EEOSAuthTranslationFlags::SetTokenFromString))
 	{
-		if (InToken.IsType<TArray<uint8>>())
+		if (!CredentialsToken.IsType<FString>())
 		{
-			const TArray<uint8>& TokenData = InToken.Get<TArray<uint8>>();
-			uint32_t InOutBufferLength = EOS_MAX_TOKEN_SIZE;
-			EOS_ByteArray_ToString(TokenData.GetData(), TokenData.Num(), TokenAnsi, &InOutBufferLength);
+			UE_LOG(LogOnlineServices, Warning, TEXT("FEOSAuthLoginOptions::Create: Failed - Expected CredentialsToken set to string for Login type: %s"), *CredentialsType.ToString());
+			return TDefaultErrorResultInternal<FEOSAuthLoginOptions>(Errors::InvalidCreds());
 		}
-		else if (InToken.IsType<FString>())
+
+		FTCHARToUTF8 Converter(*CredentialsToken.Get<FString>());
+		EOSAuthLoginOptions.TokenUtf8.SetNumZeroed(Converter.Length() + 1);
+		FCStringAnsi::Strncpy(EOSAuthLoginOptions.TokenUtf8.GetData(), Converter.Get(), EOSAuthLoginOptions.TokenUtf8.Num());
+		EOSAuthLoginOptions.CredentialsData.Token = EOSAuthLoginOptions.TokenUtf8.GetData();
+	}
+
+	if (EnumHasAnyFlags(TranslatorTraits->Flags, EEOSAuthTranslationFlags::SetTokenFromExternalAuth))
+	{
+		if (!CredentialsToken.IsType<FExternalAuthToken>())
+		{
+			UE_LOG(LogOnlineServices, Warning, TEXT("FEOSAuthLoginOptions::Create: Failed - Expected CredentialsToken set to ExternalAuthToken for Login type: %s"), *CredentialsType.ToString());
+			return TDefaultErrorResultInternal<FEOSAuthLoginOptions>(Errors::InvalidCreds());
+		}
+
+		const FExternalAuthToken& ExternalAuthToken = CredentialsToken.Get<FExternalAuthToken>();
+		const FEOSExternalAuthTranslationTraits* ExternalAuthTranslationTraits = GetExternalAuthTranslationTraits(ExternalAuthToken.Type);
+		if (ExternalAuthTranslationTraits == nullptr)
+		{
+			UE_LOG(LogOnlineServices, Warning, TEXT("FEOSAuthLoginOptions::Create: Failed - Unsuppoerted external auth type: %s"), *ExternalAuthToken.Type.ToString());
+			return TDefaultErrorResultInternal<FEOSAuthLoginOptions>(Errors::InvalidCreds());
+		}
+
+		EOSAuthLoginOptions.TokenUtf8.SetNumZeroed(ExternalAuthToken.Data.Len() + 1);
+		FCStringAnsi::Strncpy(EOSAuthLoginOptions.TokenUtf8.GetData(), TCHAR_TO_UTF8(*ExternalAuthToken.Data), EOSAuthLoginOptions.TokenUtf8.Num());
+		EOSAuthLoginOptions.CredentialsData.Token = EOSAuthLoginOptions.TokenUtf8.GetData();
+		EOSAuthLoginOptions.CredentialsData.ExternalType = ExternalAuthTranslationTraits->Type;
+	}
+
+	// todo: handle SystemAuthCredentialsOptions
+
+	// Translate scopes.
+	bool bAllScopesValid = true;
+	for (const FString& Scope : Scopes)
+	{
+		if (const EOS_EAuthScopeFlags* ScopeFlag = GetScopeFlag(Scope))
 		{
-			const FString& TokenString = InToken.Get<FString>();
-			FCStringAnsi::Strncpy(TokenAnsi, TCHAR_TO_UTF8(*TokenString), EOS_MAX_TOKEN_SIZE);
+			EOSAuthLoginOptions.ScopeFlags |= *ScopeFlag;
 		}
 		else
 		{
-			UE_LOG(LogTemp, Warning, TEXT("SetToken cannot be set with an invalid credentials token parameter. Please ensure there is valid data in the credentials token."));
+			UE_LOG(LogOnlineServices, Warning, TEXT("FEOSAuthLoginOptions::Create: Failed - Unknown scope: %s"), *Scope);
+			bAllScopesValid = false;
 		}
 	}
+	if (!bAllScopesValid)
+	{
+		return TDefaultErrorResultInternal<FEOSAuthLoginOptions>(Errors::InvalidCreds());
+	}
 
-	char TokenAnsi[EOS_MAX_TOKEN_SIZE];
-};
+	return TDefaultErrorResultInternal<FEOSAuthLoginOptions>(MoveTemp(EOSAuthLoginOptions));
+}
+
+TSharedPtr<FAccountInfoEOS> FAccountInfoRegistryEOS::Find(FPlatformUserId PlatformUserId) const
+{
+	return StaticCastSharedPtr<FAccountInfoEOS>(Super::Find(PlatformUserId));
+}
+
+TSharedPtr<FAccountInfoEOS> FAccountInfoRegistryEOS::Find(FOnlineAccountIdHandle AccountIdHandle) const
+{
+	return StaticCastSharedPtr<FAccountInfoEOS>(Super::Find(AccountIdHandle));
+}
+
+TSharedPtr<FAccountInfoEOS> FAccountInfoRegistryEOS::Find(EOS_EpicAccountId EpicAccountId) const
+{
+	FReadScopeLock Lock(IndexLock);
+	if (const TSharedRef<FAccountInfoEOS>* FoundPtr = AuthDataByEpicAccountId.Find(EpicAccountId))
+	{
+		return *FoundPtr;
+	}
+	else
+	{
+		return nullptr;
+	}
+}
 
-FEOSConnectLoginCredentialsEAS::FEOSConnectLoginCredentialsEAS(EOS_Auth_Token* InEASToken) : EASToken(InEASToken)
+TSharedPtr<FAccountInfoEOS> FAccountInfoRegistryEOS::Find(EOS_ProductUserId ProductUserId) const
 {
-	check(EASToken);
-	Type = EOS_EExternalCredentialType::EOS_ECT_EPIC;
-	Token = EASToken->AccessToken;
+	FReadScopeLock Lock(IndexLock);
+	if (const TSharedRef<FAccountInfoEOS>* FoundPtr = AuthDataByProductUserId.Find(ProductUserId))
+	{
+		return *FoundPtr;
+	}
+	else
+	{
+		return nullptr;
+	}
+}
+
+void FAccountInfoRegistryEOS::Register(const TSharedRef<FAccountInfoEOS>& AccountInfoEOS)
+{
+	FWriteScopeLock Lock(IndexLock);
+	DoRegister(AccountInfoEOS);
+}
+
+void FAccountInfoRegistryEOS::Unregister(FOnlineAccountIdHandle AccountId)
+{
+	if (TSharedPtr<FAccountInfoEOS> AccountInfoEOS = Find(AccountId))
+	{
+		FWriteScopeLock Lock(IndexLock);
+		DoUnregister(AccountInfoEOS.ToSharedRef());
+	}
+	else
+	{
+		UE_LOG(LogOnlineServices, Warning, TEXT("[FAccountInfoRegistryEOS::Unregister] Failed to find account [%s]."), *ToLogString(AccountId));
+	}
 }
 
-FEOSConnectLoginCredentialsEAS::~FEOSConnectLoginCredentialsEAS()
+void FAccountInfoRegistryEOS::DoRegister(const TSharedRef<FAccountInfo>& AccountInfo)
+{
+	TSharedRef<FAccountInfoEOS> AccountInfoEOS = StaticCastSharedRef<FAccountInfoEOS>(AccountInfo);
+
+	Super::DoRegister(AccountInfo);
+
+	if (AccountInfoEOS->EpicAccountId)
+	{
+		AuthDataByEpicAccountId.Add(AccountInfoEOS->EpicAccountId, AccountInfoEOS);
+	}
+
+	if (AccountInfoEOS->ProductUserId)
+	{
+		AuthDataByProductUserId.Add(AccountInfoEOS->ProductUserId, AccountInfoEOS);
+	}
+}
+void FAccountInfoRegistryEOS::DoUnregister(const TSharedRef<FAccountInfo>& AccountInfo)
 {
-	if (EASToken)
+	TSharedRef<FAccountInfoEOS> AccountInfoEOS = StaticCastSharedRef<FAccountInfoEOS>(AccountInfo);
+
+	Super::DoUnregister(AccountInfo);
+
+	if (AccountInfoEOS->EpicAccountId)
 	{
-		EOS_Auth_Token_Release(EASToken);
+		AuthDataByEpicAccountId.Remove(AccountInfoEOS->EpicAccountId);
+	}
+
+	if (AccountInfoEOS->ProductUserId)
+	{
+		AuthDataByProductUserId.Remove(AccountInfoEOS->ProductUserId);
 	}
 }
 
@@ -182,6 +561,8 @@ void FAuthEOSGS::Initialize()
 {
 	Super::Initialize();
 
+	RegisterHandlers();
+
 	AuthHandle = EOS_Platform_GetAuthInterface(static_cast<FOnlineServicesEOSGS&>(GetServices()).GetEOSPlatformHandle());
 	check(AuthHandle != nullptr);
 
@@ -189,417 +570,921 @@ void FAuthEOSGS::Initialize()
 	check(ConnectHandle != nullptr);
 }
 
-void SplitCredentialsType(const FString& InStr, FString& OutCredentialTypeStr, FString& OutExternalCredentialTypeStr)
+void FAuthEOSGS::PreShutdown()
 {
-	InStr.Split(FString(TEXT(":")), &OutCredentialTypeStr, &OutExternalCredentialTypeStr);
-	if (OutCredentialTypeStr.IsEmpty())
-	{
-		OutCredentialTypeStr = InStr;
-	}
-}
-
-TSharedPtr<FEOSConnectLoginCredentials> MakeConnectLoginCredentials(TOnlineAsyncOp<FAuthLogin>& InAsyncOp)
-{
-	const FAuthLogin::Params& LoginParams = InAsyncOp.GetParams();
+	Super::PreShutdown();
 
-	TSharedRef<FEOSConnectLoginCredentialsEOS> Credentials = MakeShared<FEOSConnectLoginCredentialsEOS>();
-	LexFromString(Credentials->Type, *LoginParams.CredentialsType);
-	Credentials->SetToken(LoginParams.CredentialsToken);
-	return Credentials;
+	UnregisterHandlers();
 }
 
 TOnlineAsyncOpHandle<FAuthLogin> FAuthEOSGS::Login(FAuthLogin::Params&& Params)
 {
-#if (WITH_EDITOR || !UE_BUILD_SHIPPING)
-	// Handle developer login
-	if (Params.CredentialsId.IsEmpty() && Params.CredentialsType.IsEmpty() && Params.CredentialsToken.IsType<FString>() && Params.CredentialsToken.Get<FString>().IsEmpty())
+	TOnlineAsyncOpRef<FAuthLogin> Op = GetOp<FAuthLogin>(MoveTemp(Params));
+
+	// Login in AuthEOSGS supports calling "Login" again with an updated external auth token to
+	// avoid service interruption when the connect token expires.
+
+	// Step 1: Set up operation data.
+	Op->Then([this](TOnlineAsyncOp<FAuthLogin>& InAsyncOp)
 	{
-		FString	CommandLineAuthId;
-		FString	CommandLineAuthToken;
-		FString	CommandLineAuthType;
-		FParse::Value(FCommandLine::Get(), TEXT("AUTH_LOGIN="), CommandLineAuthId);
-		FParse::Value(FCommandLine::Get(), TEXT("AUTH_PASSWORD="), CommandLineAuthToken);
-		FParse::Value(FCommandLine::Get(), TEXT("AUTH_TYPE="), CommandLineAuthType);
-		if (!CommandLineAuthId.IsEmpty() && !CommandLineAuthToken.IsEmpty() && !CommandLineAuthType.IsEmpty())
+		const FAuthLogin::Params& Params = InAsyncOp.GetParams();
+		TSharedPtr<FAccountInfoEOS> AccountInfoEOS = AccountInfoRegistryEOS.Find(Params.PlatformUserId);
+		if (!AccountInfoEOS)
 		{
-			Params.CredentialsId = MoveTemp(CommandLineAuthId);
-			Params.CredentialsToken.Emplace<FString>(MoveTemp(CommandLineAuthToken));
-			Params.CredentialsType = MoveTemp(CommandLineAuthType);
+			AccountInfoEOS = MakeShared<FAccountInfoEOS>();
+			AccountInfoEOS->PlatformUserId = Params.PlatformUserId;
+			AccountInfoEOS->LoginStatus = ELoginStatus::NotLoggedIn;
 		}
-	}
-#endif
 
-	TOnlineAsyncOpRef<FAuthLogin> Op = GetOp<FAuthLogin>(MoveTemp(Params));
-	// Are we already logged in?
-	if (GetAccountIdByPlatformUserId(Op->GetParams().PlatformUserId).IsOk())
+		// Set user auth data on operation.
+		InAsyncOp.Data.Set<TSharedRef<FAccountInfoEOS>>(AccountInfoKeyName, AccountInfoEOS.ToSharedRef());
+	})
+	// Step 2: Optionally handle EAS login.
+	.Then([this](TOnlineAsyncOp<FAuthLogin>& InAsyncOp)
 	{
-		Op->SetError(Errors::Auth::AlreadyLoggedIn());
-		return Op->GetHandle();
-	}
+		const TSharedRef<FAccountInfoEOS>& AccountInfoEOS = GetOpDataChecked<TSharedRef<FAccountInfoEOS>>(InAsyncOp, AccountInfoKeyName);
+		FAuthEOSGSLoginConfig AuthEOSGSLoginConfig;
+		LoadConfig(AuthEOSGSLoginConfig, TEXT("Login"));
+		
+		if (AuthEOSGSLoginConfig.EASAuthEnabled)
+		{
+			if (!IsOnlineStatus(AccountInfoEOS->LoginStatus))
+			{
+				TPromise<void> Promise;
+				TFuture<void> Future = Promise.GetFuture();
+
+				const FAuthLogin::Params& Params = InAsyncOp.GetParams();
+
+				FLoginEASImpl::Params LoginParams;
+				LoginParams.PlatformUserId = Params.PlatformUserId;
+				LoginParams.CredentialsType = Params.CredentialsType;
+				LoginParams.CredentialsId = Params.CredentialsId;
+				LoginParams.CredentialsToken = Params.CredentialsToken;
+				LoginParams.Scopes = !Params.Scopes.IsEmpty() ? Params.Scopes : AuthEOSGSLoginConfig.DefaultScopes;
+
+				return LoginEASImpl(LoginParams)
+				.Next([this, Promise = MoveTemp(Promise), Op = InAsyncOp.AsShared()](TDefaultErrorResult<FLoginEASImpl>&& LoginResult) mutable -> void
+				{
+					const TSharedRef<FAccountInfoEOS>& AccountInfoEOS = GetOpDataChecked<TSharedRef<FAccountInfoEOS>>(*Op, AccountInfoKeyName);
+
+					if (LoginResult.IsError())
+					{
+						UE_LOG(LogOnlineServices, Warning, TEXT("[FAuthEOSGS::Login] Failure: LoginEASImpl %s"), *LoginResult.GetErrorValue().GetLogString());
+						Op->SetError(Errors::Unknown(MoveTemp(LoginResult.GetErrorValue())));
+					}
+					else
+					{
+						// Cache EpicAccountId on successful EAS login.
+						AccountInfoEOS->EpicAccountId = LoginResult.GetOkValue().EpicAccountId;
+					}
+
+					Promise.EmplaceValue();
+				});
+
+				return Future;
+			}
+			else
+			{
+				UE_LOG(LogOnlineServices, Warning, TEXT("[FAuthEOSGS::Login] Failure: EAS login enabled and already logged in."));
+				InAsyncOp.SetError(Errors::Auth::AlreadyLoggedIn());
+			}
+		}
 
-	TOnlineChainableAsyncOp<FAuthLogin, TSharedPtr<FEOSConnectLoginCredentials>> ConnectLoginOp = [this, &Op]()
+		return MakeFulfilledPromise<void>().GetFuture();
+	})
+	// Step 3: Fetch external auth credentials for connect login.
+	.Then([this](TOnlineAsyncOp<FAuthLogin>& InAsyncOp)
 	{
-		FString LoginCredentialTypeStr, Unused;
-		SplitCredentialsType(Op->GetParams().CredentialsType, LoginCredentialTypeStr, Unused);
+		const FAuthLogin::Params& Params = InAsyncOp.GetParams();
+		const TSharedRef<FAccountInfoEOS>& AccountInfoEOS = GetOpDataChecked<TSharedRef<FAccountInfoEOS>>(InAsyncOp, AccountInfoKeyName);
+
+		TPromise<FLoginConnectImpl::Params> Promise;
+		TFuture<FLoginConnectImpl::Params> Future = Promise.GetFuture();
+
+		// When an epic account is present, use the external auth credentials of the logged in account.
+		if (AccountInfoEOS->EpicAccountId)
+		{
+			TDefaultErrorResult<FGetExternalAuthTokenImpl> AuthTokenResult = GetExternalAuthTokenImpl(FGetExternalAuthTokenImpl::Params{AccountInfoEOS->EpicAccountId});
+			if (AuthTokenResult.IsError())
+			{
+				UE_LOG(LogOnlineServices, Warning, TEXT("[FAuthEOSGS::Login] Failure: GetExternalAuthTokenImpl %s"), *AuthTokenResult.GetErrorValue().GetLogString());
+				InAsyncOp.SetError(Errors::Unknown(MoveTemp(AuthTokenResult.GetErrorValue())));
+
+				// Failed to acquire token - logout EAS.
+				LogoutEASImpl(FLogoutEASImpl::Params{ AccountInfoEOS->EpicAccountId })
+				.Next([Promise = MoveTemp(Promise)](TDefaultErrorResult<FLogoutEASImpl>&& LogoutResult) mutable -> void
+				{
+					if (LogoutResult.IsError())
+					{
+						UE_LOG(LogOnlineServices, Warning, TEXT("[FAuthEOSGS::Login] Failure: LogoutEASImpl %s"), *LogoutResult.GetErrorValue().GetLogString());
+					}
+					Promise.EmplaceValue(FLoginConnectImpl::Params{});
+				});
+
+				return Future;
+			}
 
-		// If the Credentials type is an EAS Login credential type, then login to EOS_Auth with that and get an EAS token to login to EOS_Connect with
-		EOS_ELoginCredentialType EasLoginType;
-		if (LexFromString(EasLoginType, *LoginCredentialTypeStr))
+			Promise.EmplaceValue(FLoginConnectImpl::Params{Params.PlatformUserId, MoveTemp(AuthTokenResult.GetOkValue().Token)});
+		}
+		else
 		{
-			return LoginEAS(*Op);
+			// Connect login requires external auth.
+			if (Params.CredentialsType != LoginCredentialsType::ExternalAuth || !Params.CredentialsToken.IsType<FExternalAuthToken>())
+			{
+				InAsyncOp.SetError(Errors::InvalidParams());
+				Promise.EmplaceValue(FLoginConnectImpl::Params{});
+				return Future;
+			}
+
+			Promise.EmplaceValue(FLoginConnectImpl::Params{Params.PlatformUserId, Params.CredentialsToken.Get<FExternalAuthToken>()});
 		}
 
-		// The credentials type is presumably an EOS_EExternalCredentialType, so just login to connect with that
-		return Op->Then([this](TOnlineAsyncOp<FAuthLogin>& InAsyncOp)
+		return Future;
+	})
+	// Step 4: Attempt connect login. On connect login failure handle logout of EAS.
+	.Then([this](TOnlineAsyncOp<FAuthLogin>& InAsyncOp, FLoginConnectImpl::Params&& LoginConnectParams)
+	{
+		const TSharedRef<FAccountInfoEOS>& AccountInfoEOS = GetOpDataChecked<TSharedRef<FAccountInfoEOS>>(InAsyncOp, AccountInfoKeyName);
+
+		TPromise<void> Promise;
+		TFuture<void> Future = Promise.GetFuture();
+
+		// Attempt connect login.
+		LoginConnectImpl(LoginConnectParams)
+		.Next([this, AccountInfoEOS, Op = InAsyncOp.AsShared(), Promise = MoveTemp(Promise)](TDefaultErrorResult<FLoginConnectImpl>&& LoginResult) mutable -> void
 		{
-			return MakeConnectLoginCredentials(InAsyncOp);
+			if (LoginResult.IsError())
+			{
+				UE_LOG(LogOnlineServices, Warning, TEXT("[FAuthEOSGS::Login] Failure: LoginConnectImpl %s"), *LoginResult.GetErrorValue().GetLogString());
+				Op->SetError(Errors::Unknown(MoveTemp(LoginResult.GetErrorValue())));
+
+				// Logout of EAS on login failure if necessary.
+				if (AccountInfoEOS->EpicAccountId)
+				{
+					LogoutEASImpl(FLogoutEASImpl::Params{ AccountInfoEOS->EpicAccountId })
+					.Next([Promise = MoveTemp(Promise)](TDefaultErrorResult<FLogoutEASImpl>&& LogoutResult) mutable -> void
+					{
+						if (LogoutResult.IsError())
+						{
+							UE_LOG(LogOnlineServices, Warning, TEXT("[FAuthEOSGS::Login] Failure: LogoutEASImpl %s"), *LogoutResult.GetErrorValue().GetLogString());
+						}
+						Promise.EmplaceValue();
+					});
+				}
+				else
+				{
+					Promise.EmplaceValue();
+				}
+			}
+			else
+			{
+				// Successful login.
+				AccountInfoEOS->ProductUserId = LoginResult.GetOkValue().ProductUserId;
+				Promise.EmplaceValue();
+			}
 		});
-	}();
-	// Check if the above steps completed (failed) the operation
-	if (ConnectLoginOp.GetOwningOperation().IsComplete())
-	{
-		return ConnectLoginOp.GetOwningOperation().GetHandle();
-	}
-	ConnectLoginOp.Then([this](TOnlineAsyncOp<FAuthLogin>& InAsyncOp, TSharedPtr<FEOSConnectLoginCredentials>&& InConnectLoginCredentials, TPromise<const EOS_Connect_LoginCallbackInfo*>&& Promise)
+
+		return Future;
+	})
+	// Step 5: Fetch dependent data.
+	.Then([this](TOnlineAsyncOp<FAuthLogin>& InAsyncOp)
 	{
-		EOS_Connect_LoginOptions ConnectLoginOptions = { };
-		ConnectLoginOptions.ApiVersion = EOS_CONNECT_LOGIN_API_LATEST;
-		ConnectLoginOptions.Credentials = InConnectLoginCredentials.Get();
+		const TSharedRef<FAccountInfoEOS>& AccountInfoEOS = GetOpDataChecked<TSharedRef<FAccountInfoEOS>>(InAsyncOp, AccountInfoKeyName);
+
+		if (!IsOnlineStatus(AccountInfoEOS->LoginStatus))
+		{
+			// Set DisplayName.
+			int32_t ProductUserIdBufferLength = EOS_PRODUCTUSERID_MAX_LENGTH;
+			char ProductUserIdBuffer[EOS_PRODUCTUSERID_MAX_LENGTH] = {};
+			EOS_EResult ProductUserIdResult = EOS_ProductUserId_ToString(AccountInfoEOS->ProductUserId, ProductUserIdBuffer, &ProductUserIdBufferLength);
+			if (ProductUserIdResult == EOS_EResult::EOS_Success)
+			{
+				AccountInfoEOS->Attributes.Emplace(AccountAttributeData::DisplayName, UTF8_TO_TCHAR(ProductUserIdBuffer));
+			}
+			else
+			{
+				FOnlineError ProductUserIdError(FromEOSError(ProductUserIdResult));
+				UE_LOG(LogOnlineServices, Warning, TEXT("[FAuthEOSGS::Login] Failure: EOS_ProductUserId_ToString %s"), *ProductUserIdError.GetLogString());
+				InAsyncOp.SetError(Errors::Unknown(MoveTemp(ProductUserIdError)));
+
+				// Handle EAS logout if needed.
+				if (AccountInfoEOS->EpicAccountId)
+				{
+					TPromise<void> Promise;
+					TFuture<void> Future = Promise.GetFuture();
+
+					LogoutEASImpl(FLogoutEASImpl::Params{ AccountInfoEOS->EpicAccountId })
+					.Next([Promise = MoveTemp(Promise)](TDefaultErrorResult<FLogoutEASImpl>&& LogoutResult) mutable -> void
+					{
+						if (LogoutResult.IsError())
+						{
+							UE_LOG(LogOnlineServices, Warning, TEXT("[FAuthEOSGS::Login] Failure: LogoutEASImpl %s"), *LogoutResult.GetErrorValue().GetLogString());
+						}
+						Promise.EmplaceValue();
+					});
+
+					return Future;
+				}
+			}
+		}
 
-		EOS_Async(EOS_Connect_Login, ConnectHandle, ConnectLoginOptions, MoveTemp(Promise));
+		return MakeFulfilledPromise<void>().GetFuture();
 	})
-	.Then([this](TOnlineAsyncOp<FAuthLogin>& InAsyncOp, const EOS_Connect_LoginCallbackInfo* Data, TPromise<const EOS_Connect_CreateUserCallbackInfo*>&& Promise)
+	// Step 6: bookkeeping and notifications.
+	.Then([this](TOnlineAsyncOp<FAuthLogin>& InAsyncOp)
 	{
-		UE_LOG(LogTemp, Verbose, TEXT("[FAuthEOSGS::Login] EOS_Connect_Login Result: [%s]"), *LexToString(Data->ResultCode));
+		const TSharedRef<FAccountInfoEOS>& AccountInfoEOS = GetOpDataChecked<TSharedRef<FAccountInfoEOS>>(InAsyncOp, AccountInfoKeyName);
 
-		if (Data->ResultCode == EOS_EResult::EOS_Success)
+		if (IsOnlineStatus(AccountInfoEOS->LoginStatus))
 		{
-			ProcessSuccessfulLogin(InAsyncOp, Data->LocalUserId);
+			UE_LOG(LogOnlineServices, Log, TEXT("[FAuthEOSGS::Login] Successfully refreshed logged in as [%s]"), *ToLogString(AccountInfoEOS->AccountId));
 		}
-		else if (Data->ResultCode == EOS_EResult::EOS_InvalidUser && Data->ContinuanceToken != nullptr)
+		else
 		{
-			EOS_Connect_CreateUserOptions ConnectCreateUserOptions = { };
-			ConnectCreateUserOptions.ApiVersion = EOS_CONNECT_CREATEUSER_API_LATEST;
-			ConnectCreateUserOptions.ContinuanceToken = Data->ContinuanceToken;
+			AccountInfoEOS->LoginStatus = ELoginStatus::LoggedIn;
+			AccountInfoEOS->AccountId = CreateAccountId(AccountInfoEOS->ProductUserId);
+			AccountInfoRegistryEOS.Register(AccountInfoEOS);
+
+			UE_LOG(LogOnlineServices, Log, TEXT("[FAuthEOSGS::Login] Successfully logged in as [%s]"), *ToLogString(AccountInfoEOS->AccountId));
+			OnAuthLoginStatusChangedEvent.Broadcast(FAuthLoginStatusChanged{ AccountInfoEOS, AccountInfoEOS->LoginStatus });
+		}
+
+		InAsyncOp.SetResult(FAuthLogin::Result{AccountInfoEOS});
+	})
+	.Enqueue(GetSerialQueue());
+
+	return Op->GetHandle();
+}
+
+TOnlineAsyncOpHandle<FAuthLogout> FAuthEOSGS::Logout(FAuthLogout::Params&& Params)
+{
+	TOnlineAsyncOpRef<FAuthLogout> Op = GetOp<FAuthLogout>(MoveTemp(Params));
 
-			EOS_Async(EOS_Connect_CreateUser, ConnectHandle, ConnectCreateUserOptions, MoveTemp(Promise));
+	// Step 1: Set up operation data.
+	Op->Then([this](TOnlineAsyncOp<FAuthLogout>& InAsyncOp)
+	{
+		const FAuthLogout::Params& Params = InAsyncOp.GetParams();
+		TSharedPtr<FAccountInfoEOS> AccountInfoEOS = AccountInfoRegistryEOS.Find(Params.LocalUserId);
+		if (!AccountInfoEOS)
+		{
+			InAsyncOp.SetError(Errors::InvalidUser());
 			return;
 		}
-		else
+
+		// Set user auth data on operation.
+		InAsyncOp.Data.Set<TSharedRef<FAccountInfoEOS>>(AccountInfoKeyName, AccountInfoEOS.ToSharedRef());
+	})
+	// Step 2: Delete persistent auth if needed.
+	.Then([this](TOnlineAsyncOp<FAuthLogout>& InAsyncOp)
+	{
+		const FAuthLogout::Params& Params = InAsyncOp.GetParams();
+		const TSharedRef<FAccountInfoEOS>& AccountInfoEOS = GetOpDataChecked<TSharedRef<FAccountInfoEOS>>(InAsyncOp, AccountInfoKeyName);
+
+		if (Params.bDestroyAuth && AccountInfoEOS->EpicAccountId)
 		{
-			// TODO: EAS Logout
-			InAsyncOp.SetError(Errors::Unknown()); // TODO
+			TPromise<void> Promise;
+			TFuture<void> Future = Promise.GetFuture();
+
+			EOS_Auth_DeletePersistentAuthOptions DeletePersistentAuthOptions = { 0 };
+			DeletePersistentAuthOptions.ApiVersion = EOS_AUTH_DELETEPERSISTENTAUTH_API_LATEST;
+			DeletePersistentAuthOptions.RefreshToken = nullptr; // TODO: Is this needed?  Docs say it's needed for consoles
+
+			EOS_Async(EOS_Auth_DeletePersistentAuth, AuthHandle, DeletePersistentAuthOptions,
+			[Promise = MoveTemp(Promise)](const EOS_Auth_DeletePersistentAuthCallbackInfo* Data) mutable -> void
+			{
+				if (Data->ResultCode != EOS_EResult::EOS_Success)
+				{
+					FOnlineError DeletePersistentAuthError(FromEOSError(Data->ResultCode));
+					UE_LOG(LogOnlineServices, Warning, TEXT("[FAuthEOSGS::Logout] Failure: DeletePersistentAuthResult %s"), *DeletePersistentAuthError.GetLogString());
+				}
+
+				Promise.EmplaceValue();
+			});
+
+			return Future;
 		}
 
-		Promise.EmplaceValue();
+		return MakeFulfilledPromise<void>().GetFuture();
 	})
-	.Then([this](TOnlineAsyncOp<FAuthLogin>& InAsyncOp, const EOS_Connect_CreateUserCallbackInfo* Data)
+	// Step 3: Logout of EAS if necessary.
+	.Then([this](TOnlineAsyncOp<FAuthLogout>& InAsyncOp)
 	{
-		UE_LOG(LogTemp, Verbose, TEXT("[FAuthEOSGS::Login] EOS_Connect_CreateUser Result: [%s]"), *LexToString(Data->ResultCode));
+		const TSharedRef<FAccountInfoEOS>& AccountInfoEOS = GetOpDataChecked<TSharedRef<FAccountInfoEOS>>(InAsyncOp, AccountInfoKeyName);
 
-		if (Data->ResultCode == EOS_EResult::EOS_Success)
+		if (AccountInfoEOS->EpicAccountId)
 		{
-			ProcessSuccessfulLogin(InAsyncOp, Data->LocalUserId);
+			TPromise<void> Promise;
+			TFuture<void> Future = Promise.GetFuture();
+
+			LogoutEASImpl(FLogoutEASImpl::Params{ AccountInfoEOS->EpicAccountId })
+			.Next([this, AccountInfoEOS, Promise = MoveTemp(Promise)](TDefaultErrorResult<FLogoutEASImpl>&& LogoutResult) mutable -> void
+			{
+				if (LogoutResult.IsError())
+				{
+					UE_LOG(LogOnlineServices, Warning, TEXT("[FAuthEOSGS::Logout] Failure: LogoutEASImpl %s"), *LogoutResult.GetErrorValue().GetLogString());
+				}
+				Promise.EmplaceValue();
+			});
+
+			return Future;
 		}
 		else
 		{
-			// TODO: EAS Logout
-
-			InAsyncOp.SetError(Errors::Unknown()); // TODO
+			return MakeFulfilledPromise<void>().GetFuture();
 		}
 	})
+	.Then([this](TOnlineAsyncOp<FAuthLogout>& InAsyncOp)
+	{
+		const TSharedRef<FAccountInfoEOS>& AccountInfoEOS = GetOpDataChecked<TSharedRef<FAccountInfoEOS>>(InAsyncOp, AccountInfoKeyName);
+
+		UE_LOG(LogOnlineServices, Log, TEXT("[FAuthEOSGS::Logout] Successfully logged out [%s]"), *ToLogString(AccountInfoEOS->AccountId));
+		AccountInfoEOS->LoginStatus = ELoginStatus::NotLoggedIn;
+		OnAuthLoginStatusChangedEvent.Broadcast(FAuthLoginStatusChanged{ AccountInfoEOS, AccountInfoEOS->LoginStatus });
+		AccountInfoRegistryEOS.Unregister(AccountInfoEOS->AccountId);
+
+		InAsyncOp.SetResult(FAuthLogout::Result{});
+	})
 	.Enqueue(GetSerialQueue());
 
 	return Op->GetHandle();
 }
 
-TOnlineChainableAsyncOp<FAuthLogin, TSharedPtr<FEOSConnectLoginCredentials>> FAuthEOSGS::LoginEAS(TOnlineAsyncOp<FAuthLogin>& InAsyncOp)
+TOnlineAsyncOpHandle<FAuthQueryVerifiedAuthTicket> FAuthEOSGS::QueryVerifiedAuthTicket(FAuthQueryVerifiedAuthTicket::Params&& Params)
+{
+	// Todo
+	TOnlineAsyncOpRef<FAuthQueryVerifiedAuthTicket> Operation = GetOp<FAuthQueryVerifiedAuthTicket>(MoveTemp(Params));
+	Operation->SetError(Errors::NotImplemented());
+	return Operation->GetHandle();
+}
+
+TOnlineAsyncOpHandle<FAuthCancelVerifiedAuthTicket> FAuthEOSGS::CancelVerifiedAuthTicket(FAuthCancelVerifiedAuthTicket::Params&& Params)
+{
+	// Todo
+	TOnlineAsyncOpRef<FAuthCancelVerifiedAuthTicket> Operation = GetOp<FAuthCancelVerifiedAuthTicket>(MoveTemp(Params));
+	Operation->SetError(Errors::NotImplemented());
+	return Operation->GetHandle();
+}
+
+TOnlineAsyncOpHandle<FAuthBeginVerifiedAuthSession> FAuthEOSGS::BeginVerifiedAuthSession(FAuthBeginVerifiedAuthSession::Params&& Params)
+{
+	// Todo
+	TOnlineAsyncOpRef<FAuthBeginVerifiedAuthSession> Operation = GetOp<FAuthBeginVerifiedAuthSession>(MoveTemp(Params));
+	Operation->SetError(Errors::NotImplemented());
+	return Operation->GetHandle();
+}
+
+TOnlineAsyncOpHandle<FAuthEndVerifiedAuthSession> FAuthEOSGS::EndVerifiedAuthSession(FAuthEndVerifiedAuthSession::Params&& Params)
+{
+	// Todo
+	TOnlineAsyncOpRef<FAuthEndVerifiedAuthSession> Operation = GetOp<FAuthEndVerifiedAuthSession>(MoveTemp(Params));
+	Operation->SetError(Errors::NotImplemented());
+	return Operation->GetHandle();
+}
+
+TFuture<FOnlineAccountIdHandle> FAuthEOSGS::ResolveAccountId(const FOnlineAccountIdHandle& LocalUserId, const EOS_ProductUserId ProductUserId)
+{
+	return MakeFulfilledPromise<FOnlineAccountIdHandle>(CreateAccountId(ProductUserId)).GetFuture();
+}
+
+TFuture<TArray<FOnlineAccountIdHandle>> FAuthEOSGS::ResolveAccountIds(const FOnlineAccountIdHandle& LocalUserId, const TArray<EOS_ProductUserId>& InProductUserIds)
 {
-	const FAuthLogin::Params& Params = InAsyncOp.GetParams();
-	auto FailureLambda = [](TOnlineAsyncOp<FAuthLogin>& Op) -> TSharedPtr<FEOSConnectLoginCredentials>
+	TArray<FOnlineAccountIdHandle> AccountIds;
+	AccountIds.Reserve(InProductUserIds.Num());
+	for (const EOS_ProductUserId ProductUserId : InProductUserIds)
 	{
-		return nullptr;
+		AccountIds.Emplace(CreateAccountId(ProductUserId));
+	}
+	return MakeFulfilledPromise<TArray<FOnlineAccountIdHandle>>(MoveTemp(AccountIds)).GetFuture();
+}
+
+TFunction<TFuture<FOnlineAccountIdHandle>(FOnlineAsyncOp& InAsyncOp, const EOS_ProductUserId& ProductUserId)> FAuthEOSGS::ResolveProductIdFn()
+{
+	return [this](FOnlineAsyncOp& InAsyncOp, const EOS_ProductUserId& ProductUserId)
+	{
+		const FOnlineAccountIdHandle* LocalUserIdPtr = InAsyncOp.Data.Get<FOnlineAccountIdHandle>(TEXT("LocalUserId"));
+		if (!ensure(LocalUserIdPtr))
+		{
+			return MakeFulfilledPromise<FOnlineAccountIdHandle>().GetFuture();
+		}
+		return ResolveAccountId(*LocalUserIdPtr, ProductUserId);
 	};
+}
 
-	EOS_Auth_LoginOptions LoginOptions = { };
-	LoginOptions.ApiVersion = EOS_AUTH_LOGIN_API_LATEST;
+TFunction<TFuture<TArray<FOnlineAccountIdHandle>>(FOnlineAsyncOp& InAsyncOp, const TArray<EOS_ProductUserId>& ProductUserIds)> FAuthEOSGS::ResolveProductIdsFn()
+{
+	return [this](FOnlineAsyncOp& InAsyncOp, const TArray<EOS_ProductUserId>& ProductUserIds)
+	{
+		const FOnlineAccountIdHandle* LocalUserIdPtr = InAsyncOp.Data.Get<FOnlineAccountIdHandle>(TEXT("LocalUserId"));
+		if (!ensure(LocalUserIdPtr))
+		{
+			return MakeFulfilledPromise<TArray<FOnlineAccountIdHandle>>().GetFuture();
+		}
+		return ResolveAccountIds(*LocalUserIdPtr, ProductUserIds);
+	};
+}
+
+TFuture<TDefaultErrorResult<FAuthEOSGS::FLoginEASImpl>> FAuthEOSGS::LoginEASImpl(const FLoginEASImpl::Params& LoginParams)
+{
+	TDefaultErrorResultInternal<FEOSAuthLoginOptions> LoginOptionsResult = FEOSAuthLoginOptions::Create(LoginParams.CredentialsType, LoginParams.CredentialsId, LoginParams.CredentialsToken, LoginParams.Scopes);
+	if (LoginOptionsResult.IsError())
+	{
+		return MakeFulfilledPromise<TDefaultErrorResult<FLoginEASImpl>>(MoveTemp(LoginOptionsResult.GetErrorValue())).GetFuture();
+	}
 
-	FString CredentialTypeStr;
-	FString ExternalCredentialTypeStr;
-	SplitCredentialsType(Params.CredentialsType, CredentialTypeStr, ExternalCredentialTypeStr);
+	const bool IsPersistentAuthLogin = LoginParams.CredentialsType == LoginCredentialsType::PersistentAuth;
 
-	FEOSAuthCredentials Credentials;
-	if (LexFromString(Credentials.Type, *CredentialTypeStr))
+	TPromise<TDefaultErrorResult<FLoginEASImpl>> Promise;
+	TFuture<TDefaultErrorResult<FLoginEASImpl>> Future = Promise.GetFuture();
+
+	EOS_Async(EOS_Auth_Login, AuthHandle, MoveTemp(LoginOptionsResult.GetOkValue()),
+	[AuthHandle = AuthHandle, IsPersistentAuthLogin, Promise = MoveTemp(Promise)](const EOS_Auth_LoginCallbackInfo* Data) mutable -> void
 	{
-		switch (Credentials.Type)
+		UE_LOG(LogOnlineServices, Verbose, TEXT("[FAuthEOSGS::LoginEASImpl] EOS_Auth_Login Result: [%s]"), *LexToString(Data->ResultCode));
+
+		if (Data->ResultCode == EOS_EResult::EOS_Success)
+		{
+			Promise.SetValue(TDefaultErrorResult<FLoginEASImpl>(FLoginEASImpl::Result{ Data->LocalUserId }));
+		}
+		else if (Data->ResultCode == EOS_EResult::EOS_InvalidUser && Data->ContinuanceToken != nullptr)
+		{
+			EOS_Auth_LinkAccountOptions LinkAccountOptions = {};
+			LinkAccountOptions.ApiVersion = EOS_AUTH_LINKACCOUNT_API_LATEST;
+			LinkAccountOptions.ContinuanceToken = Data->ContinuanceToken;
+
+			EOS_Async(EOS_Auth_LinkAccount, AuthHandle, LinkAccountOptions,
+			[Promise = MoveTemp(Promise)](const EOS_Auth_LinkAccountCallbackInfo* Data) mutable -> void
+			{
+				UE_LOG(LogOnlineServices, Verbose, TEXT("[FAuthEOSGS::LoginEASImpl] EOS_Auth_LinkAccount Result: [%s]"), *LexToString(Data->ResultCode));
+
+				if (Data->ResultCode == EOS_EResult::EOS_Success)
+				{
+					Promise.SetValue(TDefaultErrorResult<FLoginEASImpl>(FLoginEASImpl::Result{ Data->LocalUserId }));
+				}
+				else
+				{
+					Promise.SetValue(TDefaultErrorResult<FLoginEASImpl>(FromEOSError(Data->ResultCode)));
+				}
+			});
+		}
+		else
 		{
-		case EOS_ELoginCredentialType::EOS_LCT_ExternalAuth:
-			// If an external credential type wasn't specified, we'll grab the platform default from the configuration values
-			if (ExternalCredentialTypeStr.IsEmpty())
+			FOnlineError ResolvedError = (Data->ResultCode == EOS_EResult::EOS_InvalidAuth) ? Errors::InvalidCreds() : Errors::Unknown(FromEOSError(Data->ResultCode));
+
+			const bool bShouldRemoveCachedToken =
+				Data->ResultCode == EOS_EResult::EOS_InvalidAuth ||
+				Data->ResultCode == EOS_EResult::EOS_AccessDenied ||
+				Data->ResultCode == EOS_EResult::EOS_Auth_InvalidToken;
+
+			// Remove persistent auth credentials when they are found to be invalid.
+			if (IsPersistentAuthLogin && bShouldRemoveCachedToken)
 			{
-				FAuthEOSConfig AuthEOSConfig;
-				LoadConfig(AuthEOSConfig);
-				ExternalCredentialTypeStr = AuthEOSConfig.DefaultExternalCredentialTypeStr;
+				EOS_Auth_DeletePersistentAuthOptions DeletePersistentAuthOptions = {};
+				DeletePersistentAuthOptions.ApiVersion = EOS_AUTH_DELETEPERSISTENTAUTH_API_LATEST;
+
+				EOS_Async(EOS_Auth_DeletePersistentAuth, AuthHandle, DeletePersistentAuthOptions,
+				[ResolvedError = MoveTemp(ResolvedError), Promise = MoveTemp(Promise)](const EOS_Auth_DeletePersistentAuthCallbackInfo* Data) mutable -> void
+				{
+					if (Data->ResultCode != EOS_EResult::EOS_Success)
+					{
+						FOnlineError DeletePersistentAuthError(FromEOSError(Data->ResultCode));
+						UE_LOG(LogOnlineServices, Warning, TEXT("[FAuthEOS::LoginEASImpl] Failure: DeletePersistentAuthResult %s"), *DeletePersistentAuthError.GetLogString());
+					}
+
+					Promise.SetValue(TDefaultErrorResult<FLoginEASImpl>(MoveTemp(ResolvedError)));
+				});
 			}
-
-			EOS_EExternalCredentialType ExternalType;
-			if (LexFromString(ExternalType, *ExternalCredentialTypeStr))
+			else
 			{
-				Credentials.ExternalType = ExternalType;
+				Promise.SetValue(TDefaultErrorResult<FLoginEASImpl>(MoveTemp(ResolvedError)));
 			}
+		}
+	});
+
+	return Future;
+}
+
+TFuture<TDefaultErrorResult<FAuthEOSGS::FLogoutEASImpl>> FAuthEOSGS::LogoutEASImpl(const FLogoutEASImpl::Params& LogoutParams)
+{
+	TPromise<TDefaultErrorResult<FLogoutEASImpl>> Promise;
+	TFuture<TDefaultErrorResult<FLogoutEASImpl>> Future = Promise.GetFuture();
+
+	EOS_Auth_LogoutOptions LogoutOptions = {};
+	LogoutOptions.ApiVersion = EOS_AUTH_LOGOUT_API_LATEST;
+	LogoutOptions.LocalUserId = LogoutParams.EpicAccountId;
+
+	EOS_Async(EOS_Auth_Logout, AuthHandle, LogoutOptions,
+	[Promise = MoveTemp(Promise)](const EOS_Auth_LogoutCallbackInfo* Data) mutable -> void
+	{
+		UE_LOG(LogOnlineServices, Verbose, TEXT("[FAuthEOSGS::LogoutEASImpl] EOS_Auth_Logout Result: [%s]"), *LexToString(Data->ResultCode));
 
-			Credentials.SetToken(Params.CredentialsToken);
-			break;
-		case EOS_ELoginCredentialType::EOS_LCT_ExchangeCode:
-			// This is how the Epic launcher will pass credentials to you
-			Credentials.IdAnsi[0] = '\0';
-			FCStringAnsi::Strncpy(Credentials.TokenAnsi, TCHAR_TO_UTF8(*Params.CredentialsToken.Get<FString>()), EOS_MAX_TOKEN_SIZE);
-			break;
-		case EOS_ELoginCredentialType::EOS_LCT_Password:
-			FCStringAnsi::Strncpy(Credentials.IdAnsi, TCHAR_TO_UTF8(*Params.CredentialsId), EOS_STRING_BUFFER_LENGTH);
-			FCStringAnsi::Strncpy(Credentials.TokenAnsi, TCHAR_TO_UTF8(*Params.CredentialsToken.Get<FString>()), EOS_MAX_TOKEN_SIZE);
-			break;
-		case EOS_ELoginCredentialType::EOS_LCT_Developer:
-			// This is auth via the EOS auth tool
-			FCStringAnsi::Strncpy(Credentials.IdAnsi, TCHAR_TO_UTF8(*Params.CredentialsId), EOS_STRING_BUFFER_LENGTH);
-			FCStringAnsi::Strncpy(Credentials.TokenAnsi, TCHAR_TO_UTF8(*Params.CredentialsToken.Get<FString>()), EOS_MAX_TOKEN_SIZE);
-			break;
-		case EOS_ELoginCredentialType::EOS_LCT_AccountPortal:
-			// This is auth via the EOS Account Portal
-			Credentials.IdAnsi[0] = '\0';
-			Credentials.TokenAnsi[0] = '\0';
-			break;
-		case EOS_ELoginCredentialType::EOS_LCT_PersistentAuth:
-			// This is auth via stored credentials in EOS
-			Credentials.Id = nullptr;
-			Credentials.Token = nullptr;
-			break;
-		default:
-			UE_LOG(LogTemp, Warning, TEXT("Unsupported CredentialsType=[%s]"), *Params.CredentialsType);
-			InAsyncOp.SetError(Errors::Unknown()); // TODO
-			return InAsyncOp.Then(FailureLambda);
+		if (Data->ResultCode == EOS_EResult::EOS_Success)
+		{
+			Promise.SetValue(TDefaultErrorResult<FLogoutEASImpl>(FLogoutEASImpl::Result{}));
+		}
+		else
+		{
+			Promise.SetValue(TDefaultErrorResult<FLogoutEASImpl>(FromEOSError(Data->ResultCode)));
 		}
+	});
+
+	return Future;
+}
+
+TDefaultErrorResult<FAuthEOSGS::FGetExternalAuthTokenImpl> FAuthEOSGS::GetExternalAuthTokenImpl(const FGetExternalAuthTokenImpl::Params& Params)
+{
+	EOS_Auth_CopyIdTokenOptions CopyIdTokenOptions = {};
+	CopyIdTokenOptions.ApiVersion = EOS_AUTH_COPYIDTOKEN_API_LATEST;
+	CopyIdTokenOptions.AccountId = Params.EpicAccountId;
+
+	EOS_Auth_IdToken* IdToken = nullptr;
+	EOS_EResult Result = EOS_Auth_CopyIdToken(AuthHandle, &CopyIdTokenOptions, &IdToken);
+	if (Result == EOS_EResult::EOS_Success)
+	{
+		ON_SCOPE_EXIT
+		{
+			EOS_Auth_IdToken_Release(IdToken);
+		};
+
+		FExternalAuthToken ExternalAuthToken;
+		ExternalAuthToken.Type = ExternalLoginType::EpicIdToken;
+		ExternalAuthToken.Data = UTF8_TO_TCHAR(IdToken->JsonWebToken);
+		return TDefaultErrorResult<FGetExternalAuthTokenImpl>(FGetExternalAuthTokenImpl::Result{ MoveTemp(ExternalAuthToken) });
 	}
 	else
 	{
-		UE_LOG(LogTemp, Warning, TEXT("Invalid CredentialsType=[%s]"), *Params.CredentialsType);
-		InAsyncOp.SetError(Errors::Unknown()); // TODO
-		return InAsyncOp.Then(FailureLambda);
+		return TDefaultErrorResult<FGetExternalAuthTokenImpl>(FromEOSError(Result));
 	}
+}
 
-	return InAsyncOp.Then([this, LoginOptions, Credentials](TOnlineAsyncOp<FAuthLogin>& InAsyncOp, TPromise<const EOS_Auth_LoginCallbackInfo*>&& Promise) mutable
+TFuture<TDefaultErrorResult<FAuthEOSGS::FLoginConnectImpl>> FAuthEOSGS::LoginConnectImpl(const FLoginConnectImpl::Params& LoginParams)
+{
+	TDefaultErrorResultInternal<FEOSConnectLoginOptions> LoginOptionsResult = FEOSConnectLoginOptions::Create(LoginParams.PlatformUserId, LoginParams.ExternalAuthToken);
+	if (LoginOptionsResult.IsError())
 	{
-		LoginOptions.Credentials = &Credentials;
-		EOS_Async(EOS_Auth_Login, AuthHandle, LoginOptions, MoveTemp(Promise));
-	})
-	.Then([this](TOnlineAsyncOp<FAuthLogin>& InAsyncOp, const EOS_Auth_LoginCallbackInfo* Data, TPromise<const EOS_Auth_LinkAccountCallbackInfo*>&& Promise)
+		return MakeFulfilledPromise<TDefaultErrorResult<FLoginConnectImpl>>(MoveTemp(LoginOptionsResult.GetErrorValue())).GetFuture();
+	}
+
+	TPromise<TDefaultErrorResult<FLoginConnectImpl>> Promise;
+	TFuture<TDefaultErrorResult<FLoginConnectImpl>> Future = Promise.GetFuture();
+
+	EOS_Async(EOS_Connect_Login, ConnectHandle, MoveTemp(LoginOptionsResult.GetOkValue()),
+	[ConnectHandle = ConnectHandle, Promise = MoveTemp(Promise)](const EOS_Connect_LoginCallbackInfo* Data) mutable -> void
 	{
-		UE_LOG(LogTemp, Verbose, TEXT("[FAuthEOS::Login] EOS_Auth_Login Result: [%s]"), *LexToString(Data->ResultCode));
+		UE_LOG(LogOnlineServices, Verbose, TEXT("[FAuthEOSGS::LoginConnect] EOS_Connect_Login Result: [%s]"), *LexToString(Data->ResultCode));
 
 		if (Data->ResultCode == EOS_EResult::EOS_Success)
 		{
-			// We cache the Epic Account Id to use it in later stages of the login process
-			InAsyncOp.Data.Set(TEXT("EpicAccountId"), Data->LocalUserId);
-
-			// A success means the account is already linked, we'll process the NULL result in the next step
+			Promise.SetValue(TDefaultErrorResult<FLoginConnectImpl>(FLoginConnectImpl::Result{ Data->LocalUserId }));
 		}
 		else if (Data->ResultCode == EOS_EResult::EOS_InvalidUser && Data->ContinuanceToken != nullptr)
 		{
-			EOS_Auth_LinkAccountOptions LinkAccountOptions = {};
-			LinkAccountOptions.ApiVersion = EOS_AUTH_LINKACCOUNT_API_LATEST;
-			LinkAccountOptions.ContinuanceToken = Data->ContinuanceToken;
+			EOS_Connect_CreateUserOptions ConnectCreateUserOptions = { };
+			ConnectCreateUserOptions.ApiVersion = EOS_CONNECT_CREATEUSER_API_LATEST;
+			ConnectCreateUserOptions.ContinuanceToken = Data->ContinuanceToken;
 
-			EOS_Async(EOS_Auth_LinkAccount, AuthHandle, LinkAccountOptions, MoveTemp(Promise));
-			return;
+			EOS_Async(EOS_Connect_CreateUser, ConnectHandle, ConnectCreateUserOptions,
+			[Promise = MoveTemp(Promise)](const EOS_Connect_CreateUserCallbackInfo* Data) mutable -> void
+			{
+				UE_LOG(LogOnlineServices, Verbose, TEXT("[FAuthEOSGS::LoginConnect] EOS_Connect_CreateUser Result: [%s]"), *LexToString(Data->ResultCode));
+
+				if (Data->ResultCode == EOS_EResult::EOS_Success)
+				{
+					Promise.SetValue(TDefaultErrorResult<FLoginConnectImpl>(FLoginConnectImpl::Result{ Data->LocalUserId }));
+				}
+				else
+				{
+					Promise.SetValue(TDefaultErrorResult<FLoginConnectImpl>(FromEOSError(Data->ResultCode)));
+				}
+			});
 		}
 		else
 		{
-			FOnlineError Error = Errors::Unknown();
-			if (Data->ResultCode == EOS_EResult::EOS_InvalidAuth)
-			{
-				Error = Errors::InvalidCreds();
-			}
+			Promise.SetValue(TDefaultErrorResult<FLoginConnectImpl>(FromEOSError(Data->ResultCode)));
+		}
+	});
+
+	return Future;
+}
+
+TOnlineAsyncOpHandle<FAuthEOSGS::FConnectLoginRecoveryImpl> FAuthEOSGS::ConnectLoginRecoveryImplOp(FConnectLoginRecoveryImpl::Params&& Params)
+{
+	TOnlineAsyncOpRef<FConnectLoginRecoveryImpl> Op = GetOp<FConnectLoginRecoveryImpl>(MoveTemp(Params));
+
+	// Step 1: Setup operation data.
+	Op->Then([this](TOnlineAsyncOp<FConnectLoginRecoveryImpl>& InAsyncOp)
+	{
+		const FConnectLoginRecoveryImpl::Params& Params = InAsyncOp.GetParams();
+		TSharedPtr<FAccountInfoEOS> AccountInfoEOS = AccountInfoRegistryEOS.Find(Params.LocalUserId);
+		if (!AccountInfoEOS)
+		{
+			InAsyncOp.SetError(Errors::InvalidUser());
+			return;
+		}
+
+		if (AccountInfoEOS->LoginStatus != ELoginStatus::LoggedInReducedFunctionality)
+		{
+			InAsyncOp.SetError(Errors::Unknown());
+			return;
+		}
+
+		if (!ensure(AccountInfoEOS->EpicAccountId))
+		{
+			InAsyncOp.SetError(Errors::Unknown());
+			return;
+		}
+
+		// Set user auth data on operation.
+		InAsyncOp.Data.Set<TSharedRef<FAccountInfoEOS>>(AccountInfoKeyName, AccountInfoEOS.ToSharedRef());
+	})
+	// Step 2: Acquire EAS external auth token.
+	.Then([this](TOnlineAsyncOp<FConnectLoginRecoveryImpl>& InAsyncOp)
+	{
+		const TSharedRef<FAccountInfoEOS>& AccountInfoEOS = GetOpDataChecked<TSharedRef<FAccountInfoEOS>>(InAsyncOp, AccountInfoKeyName);
+
+		TDefaultErrorResult<FGetExternalAuthTokenImpl> AuthTokenResult = GetExternalAuthTokenImpl(FGetExternalAuthTokenImpl::Params{ AccountInfoEOS->EpicAccountId });
+		if (AuthTokenResult.IsError())
+		{
+			UE_LOG(LogOnlineServices, Warning, TEXT("[FAuthEOSGS::ConnectLoginRecoveryImplOp] Failure: GetExternalAuthTokenImpl %s"), *AuthTokenResult.GetErrorValue().GetLogString());
+			InAsyncOp.SetError(Errors::Unknown(MoveTemp(AuthTokenResult.GetErrorValue())));
 
-			InAsyncOp.SetError(MoveTemp(Error));
+			// Reinitialize recovery timer.
+			InitializeConnectLoginRecoveryTimer(AccountInfoEOS);
+			return FLoginConnectImpl::Params{};
 		}
 
-		Promise.EmplaceValue();
+		return FLoginConnectImpl::Params{ AccountInfoEOS->PlatformUserId, MoveTemp(AuthTokenResult.GetOkValue().Token) };
 	})
-	.Then([this](TOnlineAsyncOp<FAuthLogin>& InAsyncOp, const EOS_Auth_LinkAccountCallbackInfo* Data) -> TFuture<TSharedPtr<FEOSConnectLoginCredentials>>
+	// Step 3: Refresh connect login.
+	.Then([this](TOnlineAsyncOp<FConnectLoginRecoveryImpl>& InAsyncOp, FLoginConnectImpl::Params&& LoginConnectParams)
 	{
-		UE_LOG(LogTemp, Verbose, TEXT("[FAuthEOS::Login] EOS_Auth_LinkAccount Result: [%s]"), (Data == nullptr) ? TEXT("Null") : *LexToString(Data->ResultCode));
+		const TSharedRef<FAccountInfoEOS>& AccountInfoEOS = GetOpDataChecked<TSharedRef<FAccountInfoEOS>>(InAsyncOp, AccountInfoKeyName);
 
-		// If Data is NULL, it means the account was already linked
-		if (Data == nullptr || Data->ResultCode == EOS_EResult::EOS_Success)
+		TPromise<void> Promise;
+		TFuture<void> Future = Promise.GetFuture();
+
+		// Attempt connect login.
+		LoginConnectImpl(LoginConnectParams)
+		.Next([this, AccountInfoEOS, Op = InAsyncOp.AsShared(), Promise = MoveTemp(Promise)](TDefaultErrorResult<FLoginConnectImpl>&& LoginResult) mutable -> void
 		{
-			EOS_EpicAccountId EpicAccountId = nullptr;
-			if (Data == nullptr)
+			if (LoginResult.IsError())
 			{
-				EpicAccountId = *InAsyncOp.Data.Get<EOS_EpicAccountId>(TEXT("EpicAccountId"));
+				UE_LOG(LogOnlineServices, Warning, TEXT("[FAuthEOSGS::ConnectLoginRecoveryImplOp] Failure: LoginConnectImpl %s"), *LoginResult.GetErrorValue().GetLogString());
+				Op->SetError(Errors::Unknown(MoveTemp(LoginResult.GetErrorValue())));
+
+				// Reinitialize recovery timer.
+				InitializeConnectLoginRecoveryTimer(AccountInfoEOS);
+				Promise.EmplaceValue();
 			}
 			else
 			{
-				// We cache the Epic Account Id to use it in later stages of the login process
-				InAsyncOp.Data.Set(TEXT("EpicAccountId"), Data->LocalUserId);
-				EpicAccountId = Data->LocalUserId;
+				// Successful login.
+				Promise.EmplaceValue();
 			}
+		});
+
+		return Future;
+	})
+	// Step 4: Update cache and notify.
+	.Then([this](TOnlineAsyncOp<FConnectLoginRecoveryImpl>& InAsyncOp)
+	{
+		const TSharedRef<FAccountInfoEOS>& AccountInfoEOS = GetOpDataChecked<TSharedRef<FAccountInfoEOS>>(InAsyncOp, AccountInfoKeyName);
+
+		UE_LOG(LogOnlineServices, Log, TEXT("[FAuthEOSGS::ConnectLoginRecoveryImplOp] User has restored online capability [%s]"), *ToLogString(AccountInfoEOS->AccountId));
+		AccountInfoEOS->LoginStatus = ELoginStatus::LoggedIn;
+		OnAuthLoginStatusChangedEvent.Broadcast(FAuthLoginStatusChanged{ AccountInfoEOS, AccountInfoEOS->LoginStatus });
+
+		InAsyncOp.SetResult(FConnectLoginRecoveryImpl::Result{});
+	})
+	.Enqueue(GetSerialQueue());
+
+	return Op->GetHandle();
+}
 
-			// On success, attempt Connect Login
-			EOS_Auth_Token* AuthToken = nullptr;
-			EOS_Auth_CopyUserAuthTokenOptions CopyOptions = { };
-			CopyOptions.ApiVersion = EOS_AUTH_COPYUSERAUTHTOKEN_API_LATEST;
+TOnlineAsyncOpHandle<FAuthEOSGS::FHandleConnectLoginStatusChangedImpl> FAuthEOSGS::HandleConnectLoginStatusChangedImplOp(FHandleConnectLoginStatusChangedImpl::Params&& Params)
+{
+	TOnlineAsyncOpRef<FHandleConnectLoginStatusChangedImpl> Op = GetOp<FHandleConnectLoginStatusChangedImpl>(MoveTemp(Params));
 
-			EOS_EResult CopyResult = EOS_Auth_CopyUserAuthToken(AuthHandle, &CopyOptions, EpicAccountId, &AuthToken);
+	// Step 1: Set up operation data.
+	Op->Then([this](TOnlineAsyncOp<FHandleConnectLoginStatusChangedImpl>& InAsyncOp)
+	{
+		const FHandleConnectLoginStatusChangedImpl::Params& Params = InAsyncOp.GetParams();
+		TSharedPtr<FAccountInfoEOS> AccountInfoEOS = AccountInfoRegistryEOS.Find(Params.LocalUserId);
+		if (!AccountInfoEOS)
+		{
+			InAsyncOp.SetError(Errors::InvalidUser());
+			return;
+		}
 
-			UE_LOG(LogTemp, Verbose, TEXT("[FAuthEOS::Login] EOS_Auth_CopyUserAuthToken Result: [%s]"), *LexToString(CopyResult));
+		// Set user auth data on operation.
+		InAsyncOp.Data.Set<TSharedRef<FAccountInfoEOS>>(AccountInfoKeyName, AccountInfoEOS.ToSharedRef());
+	})
+	// Step 2: Update status and notify.
+	.Then([this](TOnlineAsyncOp<FHandleConnectLoginStatusChangedImpl>& InAsyncOp)
+	{
+		const FHandleConnectLoginStatusChangedImpl::Params& Params = InAsyncOp.GetParams();
+		const TSharedRef<FAccountInfoEOS>& AccountInfoEOS = GetOpDataChecked<TSharedRef<FAccountInfoEOS>>(InAsyncOp, AccountInfoKeyName);
 
-			if (CopyResult == EOS_EResult::EOS_Success)
+		// Check if the user was kicked offline.
+		if (Params.CurrentStatus == EOS_ELoginStatus::EOS_LS_NotLoggedIn)
+		{
+			// Check if user can be brought back online automatically.
+			if (AccountInfoEOS->EpicAccountId)
 			{
-				TSharedPtr<FEOSConnectLoginCredentials> ConnectLoginCredentials = MakeShared<FEOSConnectLoginCredentialsEAS>(AuthToken);
-				return MakeFulfilledPromise<TSharedPtr<FEOSConnectLoginCredentials>>(MoveTemp(ConnectLoginCredentials)).GetFuture();
+				UE_LOG(LogOnlineServices, Warning, TEXT("[FAuthEOSGS::HandleConnectLoginStatusChangedImplOp] User has reduced online capability [%s]"), *ToLogString(AccountInfoEOS->AccountId));
+				AccountInfoEOS->LoginStatus = ELoginStatus::LoggedInReducedFunctionality;
+				OnAuthLoginStatusChangedEvent.Broadcast(FAuthLoginStatusChanged{ AccountInfoEOS, AccountInfoEOS->LoginStatus });
+
+				// Auth has now entered degraded state for the user. Set timer to periodically attempt to reestablish full auth connection.
+				InitializeConnectLoginRecoveryTimer(AccountInfoEOS);
 			}
 			else
 			{
-				// TODO: EAS Logout
-				InAsyncOp.SetError(Errors::Unknown()); // TODO
+				UE_LOG(LogOnlineServices, Warning, TEXT("[FAuthEOSGS::HandleConnectLoginStatusChangedImplOp] User is offline [%s]"), *ToLogString(AccountInfoEOS->AccountId));
+				AccountInfoEOS->LoginStatus = ELoginStatus::NotLoggedIn;
+				OnAuthLoginStatusChangedEvent.Broadcast(FAuthLoginStatusChanged{ AccountInfoEOS, AccountInfoEOS->LoginStatus });
+				AccountInfoRegistryEOS.Unregister(AccountInfoEOS->AccountId);
 			}
 		}
-		else
-		{
-			InAsyncOp.SetError(Errors::Unknown()); // TODO
-		}
 
-		return MakeFulfilledPromise<TSharedPtr<FEOSConnectLoginCredentials>>(nullptr).GetFuture();
-	});
+		InAsyncOp.SetResult(FHandleConnectLoginStatusChangedImpl::Result{});
+	})
+	.Enqueue(GetSerialQueue());
+
+	return Op->GetHandle();
 }
 
-void FAuthEOSGS::ProcessSuccessfulLogin(TOnlineAsyncOp<FAuthLogin>& InAsyncOp, const EOS_ProductUserId ProductUserId)
+TOnlineAsyncOpHandle<FAuthEOSGS::FHandleConnectAuthNotifyExpirationImpl> FAuthEOSGS::HandleConnectAuthNotifyExpirationImplOp(FHandleConnectAuthNotifyExpirationImpl::Params&& Params)
 {
-	const FOnlineAccountIdHandle LocalUserId = CreateAccountId(ProductUserId);
+	TOnlineAsyncOpRef<FHandleConnectAuthNotifyExpirationImpl> Op = GetOp<FHandleConnectAuthNotifyExpirationImpl>(MoveTemp(Params));
 
-	UE_LOG(LogTemp, Verbose, TEXT("[FAuthEOSGS::Login] Successfully logged in as [%s]"), *ToLogString(LocalUserId));
+	// Step 1: Set up operation data.
+	Op->Then([this](TOnlineAsyncOp<FHandleConnectAuthNotifyExpirationImpl>& InAsyncOp)
+	{
+		const FHandleConnectAuthNotifyExpirationImpl::Params& Params = InAsyncOp.GetParams();
+		TSharedPtr<FAccountInfoEOS> AccountInfoEOS = AccountInfoRegistryEOS.Find(Params.LocalUserId);
+		if (!AccountInfoEOS)
+		{
+			InAsyncOp.SetError(Errors::InvalidUser());
+			return;
+		}
 
-	TSharedRef<FAccountInfoEOS> AccountInfo = MakeShared<FAccountInfoEOS>();
-	AccountInfo->PlatformUserId = InAsyncOp.GetParams().PlatformUserId;
-	AccountInfo->UserId = LocalUserId;
-	AccountInfo->LoginStatus = ELoginStatus::LoggedIn;
+		if (AccountInfoEOS->LoginStatus != ELoginStatus::LoggedIn)
+		{
+			InAsyncOp.SetError(Errors::NotLoggedIn());
+			return;
+		}
 
-	check(!AccountInfos.Contains(LocalUserId));
-	AccountInfos.Emplace(LocalUserId, AccountInfo);
+		// Set user auth data on operation.
+		InAsyncOp.Data.Set<TSharedRef<FAccountInfoEOS>>(AccountInfoKeyName, AccountInfoEOS.ToSharedRef());
+	})
+	// Step 2: Automatically start recovery or notify user to start recovery.
+	.Then([this](TOnlineAsyncOp<FHandleConnectAuthNotifyExpirationImpl>& InAsyncOp)
+	{
+		const FHandleConnectAuthNotifyExpirationImpl::Params& Params = InAsyncOp.GetParams();
+		const TSharedRef<FAccountInfoEOS>& AccountInfoEOS = GetOpDataChecked<TSharedRef<FAccountInfoEOS>>(InAsyncOp, AccountInfoKeyName);
 
-	FAuthLogin::Result Result = { AccountInfo };
-	InAsyncOp.SetResult(MoveTemp(Result));
+		// When EAS is logged in use it to refresh the connect login.
+		if (AccountInfoEOS->EpicAccountId)
+		{
+			// Auth has now entered degraded state for the user. Set timer to periodically attempt to reestablish full auth connection.
+			InitializeConnectLoginRecoveryTimer(AccountInfoEOS);
+		}
+		else
+		{
+			// No EAS login, notify user code to refresh connect login using external auth.
+			OnAuthPendingAuthExpirationEvent.Broadcast(FAuthPendingAuthExpiration{ AccountInfoEOS });
+		}
 
-	FLoginStatusChanged EventParameters;
-	EventParameters.LocalUserId = LocalUserId;
-	EventParameters.PreviousStatus = ELoginStatus::NotLoggedIn;
-	EventParameters.CurrentStatus = ELoginStatus::LoggedIn;
-	OnLoginStatusChangedEvent.Broadcast(EventParameters);
+		InAsyncOp.SetResult(FHandleConnectAuthNotifyExpirationImpl::Result{});
+	})
+	.Enqueue(GetSerialQueue());
+
+	return Op->GetHandle();
 }
 
-TOnlineAsyncOpHandle<FAuthLogout> FAuthEOSGS::Logout(FAuthLogout::Params&& Params)
+TOnlineAsyncOpHandle<FAuthEOSGS::FHandleEASLoginStatusChangedImpl> FAuthEOSGS::HandleEASLoginStatusChangedImplOp(FHandleEASLoginStatusChangedImpl::Params&& Params)
 {
-	const FOnlineAccountIdHandle LocalUserId = Params.LocalUserId;
-	TOnlineAsyncOpRef<FAuthLogout> Op = GetOp<FAuthLogout>(MoveTemp(Params));
+	TOnlineAsyncOpRef<FHandleEASLoginStatusChangedImpl> Op = GetOp<FHandleEASLoginStatusChangedImpl>(MoveTemp(Params));
 
-	if (!ValidateOnlineId(LocalUserId))
+	// Step 1: Set up operation data.
+	Op->Then([this](TOnlineAsyncOp<FHandleEASLoginStatusChangedImpl>& InAsyncOp)
 	{
-		Op->SetError(Errors::InvalidUser());
-		return Op->GetHandle();
-	}
+		const FHandleEASLoginStatusChangedImpl::Params& Params = InAsyncOp.GetParams();
+		TSharedPtr<FAccountInfoEOS> AccountInfoEOS = AccountInfoRegistryEOS.Find(Params.LocalUserId);
+		if (!AccountInfoEOS)
+		{
+			InAsyncOp.SetError(Errors::InvalidUser());
+			return;
+		}
 
-	Op->SetResult(FAuthLogout::Result());
-	return Op->GetHandle();
-}
+		if (AccountInfoEOS->LoginStatus != ELoginStatus::LoggedIn)
+		{
+			InAsyncOp.SetError(Errors::NotLoggedIn());
+			return;
+		}
 
-TOnlineResult<FAuthGetAccountByPlatformUserId> FAuthEOSGS::GetAccountByPlatformUserId(FAuthGetAccountByPlatformUserId::Params&& Params)
-{
-	TResult<FOnlineAccountIdHandle, FOnlineError> PlatformUserIdResult = GetAccountIdByPlatformUserId(Params.PlatformUserId);
-	if (PlatformUserIdResult.IsOk())
-	{
-		FAuthGetAccountByPlatformUserId::Result Result = { AccountInfos.FindChecked(PlatformUserIdResult.GetOkValue()) };
-		return TOnlineResult<FAuthGetAccountByPlatformUserId>(Result);
-	}
-	else
+		// Set user auth data on operation.
+		InAsyncOp.Data.Set<TSharedRef<FAccountInfoEOS>>(AccountInfoKeyName, AccountInfoEOS.ToSharedRef());
+	})
+	// Step 2: Update status and notify.
+	.Then([this](TOnlineAsyncOp<FHandleEASLoginStatusChangedImpl>& InAsyncOp)
 	{
-		return TOnlineResult<FAuthGetAccountByPlatformUserId>(PlatformUserIdResult.GetErrorValue());
-	}
+		const FHandleEASLoginStatusChangedImpl::Params& Params = InAsyncOp.GetParams();
+		const TSharedRef<FAccountInfoEOS>& AccountInfoEOS = GetOpDataChecked<TSharedRef<FAccountInfoEOS>>(InAsyncOp, AccountInfoKeyName);
+
+		// Check if the user was kicked offline.
+		if (Params.CurrentStatus == EOS_ELoginStatus::EOS_LS_NotLoggedIn)
+		{
+			UE_LOG(LogOnlineServices, Warning, TEXT("[FAuthEOSGS::HandleEASLoginStatusChangedImplOp] User is now offline [%s]"), *ToLogString(AccountInfoEOS->AccountId));
+			AccountInfoEOS->LoginStatus = ELoginStatus::NotLoggedIn;
+			OnAuthLoginStatusChangedEvent.Broadcast(FAuthLoginStatusChanged{ AccountInfoEOS, AccountInfoEOS->LoginStatus });
+			AccountInfoRegistryEOS.Unregister(AccountInfoEOS->AccountId);
+		}
+
+		InAsyncOp.SetResult(FHandleEASLoginStatusChangedImpl::Result{});
+	})
+	.Enqueue(GetSerialQueue());
+
+	return Op->GetHandle();
 }
 
-TOnlineResult<FAuthGetAccountByAccountId> FAuthEOSGS::GetAccountByAccountId(FAuthGetAccountByAccountId::Params&& Params)
+void FAuthEOSGS::RegisterHandlers()
 {
-	if (TSharedRef<FAccountInfoEOS>* const FoundAccount = AccountInfos.Find(Params.LocalUserId))
-	{
-		FAuthGetAccountByAccountId::Result Result = { *FoundAccount };
-		return TOnlineResult<FAuthGetAccountByAccountId>(Result);
-	}
-	else
-	{
-		// TODO: proper error
-		return TOnlineResult<FAuthGetAccountByAccountId>(Errors::Unknown());
-	}
+	// Register for EOS connect connection status updates.
+	OnConnectLoginStatusChangedEOSEventRegistration = EOS_RegisterComponentEventHandler(
+		this,
+		ConnectHandle,
+		EOS_CONNECT_ADDNOTIFYLOGINSTATUSCHANGED_API_LATEST,
+		&EOS_Connect_AddNotifyLoginStatusChanged,
+		&EOS_Connect_RemoveNotifyLoginStatusChanged,
+		&FAuthEOSGS::OnConnectLoginStatusChanged);
+
+	// Notification of a pending external token expiration ~10 minutes.
+	OnConnectAuthNotifyExpirationEOSEventRegistration = EOS_RegisterComponentEventHandler(
+		this,
+		ConnectHandle,
+		EOS_CONNECT_ONAUTHEXPIRATIONCALLBACK_API_LATEST,
+		&EOS_Connect_AddNotifyAuthExpiration,
+		&EOS_Connect_RemoveNotifyAuthExpiration,
+		&FAuthEOSGS::OnConnectAuthNotifyExpiration);
+
+	// Register for EAS connection status updates.
+	OnConnectAuthNotifyExpirationEOSEventRegistration = EOS_RegisterComponentEventHandler(
+		this,
+		AuthHandle,
+		EOS_AUTH_ADDNOTIFYLOGINSTATUSCHANGED_API_LATEST,
+		&EOS_Auth_AddNotifyLoginStatusChanged,
+		&EOS_Auth_RemoveNotifyLoginStatusChanged,
+		&FAuthEOSGS::OnEASLoginStatusChanged);
 }
 
-bool FAuthEOSGS::IsLoggedIn(const FOnlineAccountIdHandle& AccountId) const
+void FAuthEOSGS::UnregisterHandlers()
 {
-	// TODO:  More logic?
-	return AccountInfos.Contains(AccountId);
+	OnConnectLoginStatusChangedEOSEventRegistration = nullptr;
+	OnConnectAuthNotifyExpirationEOSEventRegistration = nullptr;
+	OnConnectAuthNotifyExpirationEOSEventRegistration = nullptr;
 }
 
-TResult<FOnlineAccountIdHandle, FOnlineError> FAuthEOSGS::GetAccountIdByPlatformUserId(FPlatformUserId PlatformUserId) const
+void FAuthEOSGS::OnConnectLoginStatusChanged(const EOS_Connect_LoginStatusChangedCallbackInfo* Data)
 {
-	for (const TPair<FOnlineAccountIdHandle, TSharedRef<FAccountInfoEOS>>& AccountPair : AccountInfos)
-	{
-		if (AccountPair.Value->PlatformUserId == PlatformUserId)
-		{
-			TResult<FOnlineAccountIdHandle, FOnlineError> Result(AccountPair.Key);
-			return Result;
-		}
-	}
-	TResult<FOnlineAccountIdHandle, FOnlineError> Result(Errors::Unknown()); // TODO: error code
-	return Result;
+	HandleConnectLoginStatusChangedImplOp(FHandleConnectLoginStatusChangedImpl::Params{ Data->LocalUserId, Data->PreviousStatus, Data->CurrentStatus });
 }
 
-TFuture<FOnlineAccountIdHandle> FAuthEOSGS::ResolveAccountId(const FOnlineAccountIdHandle& LocalUserId, const EOS_ProductUserId ProductUserId)
+void FAuthEOSGS::OnConnectAuthNotifyExpiration(const EOS_Connect_AuthExpirationCallbackInfo* Data)
 {
-	return MakeFulfilledPromise<FOnlineAccountIdHandle>(CreateAccountId(ProductUserId)).GetFuture();
+	HandleConnectAuthNotifyExpirationImplOp(FHandleConnectAuthNotifyExpirationImpl::Params{ Data->LocalUserId });
 }
 
-TFuture<TArray<FOnlineAccountIdHandle>> FAuthEOSGS::ResolveAccountIds(const FOnlineAccountIdHandle& LocalUserId, const TArray<EOS_ProductUserId>& InProductUserIds)
+void FAuthEOSGS::OnEASLoginStatusChanged(const EOS_Auth_LoginStatusChangedCallbackInfo* Data)
 {
-	TArray<FOnlineAccountIdHandle> AccountIds;
-	AccountIds.Reserve(InProductUserIds.Num());
-	for (const EOS_ProductUserId ProductUserId : InProductUserIds)
-	{
-		AccountIds.Emplace(CreateAccountId(ProductUserId));
-	}
-	return MakeFulfilledPromise<TArray<FOnlineAccountIdHandle>>(MoveTemp(AccountIds)).GetFuture();
+	HandleEASLoginStatusChangedImplOp(FHandleEASLoginStatusChangedImpl::Params{ Data->LocalUserId, Data->PrevStatus, Data->CurrentStatus });
 }
 
-TFunction<TFuture<FOnlineAccountIdHandle>(FOnlineAsyncOp& InAsyncOp, const EOS_ProductUserId& ProductUserId)> FAuthEOSGS::ResolveProductIdFn()
+const FAccountInfoRegistry& FAuthEOSGS::GetAccountInfoRegistry() const
 {
-	return [this](FOnlineAsyncOp& InAsyncOp, const EOS_ProductUserId& ProductUserId)
-	{
-		const FOnlineAccountIdHandle* LocalUserIdPtr = InAsyncOp.Data.Get<FOnlineAccountIdHandle>(TEXT("LocalUserId"));
-		if (!ensure(LocalUserIdPtr))
-		{
-			return MakeFulfilledPromise<FOnlineAccountIdHandle>().GetFuture();
-		}
-		return ResolveAccountId(*LocalUserIdPtr, ProductUserId);
-	};
+	return AccountInfoRegistryEOS;
 }
 
-TFunction<TFuture<TArray<FOnlineAccountIdHandle>>(FOnlineAsyncOp& InAsyncOp, const TArray<EOS_ProductUserId>& ProductUserIds)> FAuthEOSGS::ResolveProductIdsFn()
+void FAuthEOSGS::InitializeConnectLoginRecoveryTimer(const TSharedRef<FAccountInfoEOS>& AccountInfoEOS)
 {
-	return [this](FOnlineAsyncOp& InAsyncOp, const TArray<EOS_ProductUserId>& ProductUserIds)
+	FAuthEOSGSLoginRecoveryConfig AuthEOSGSLoginRecoveryConfig;
+	LoadConfig(AuthEOSGSLoginRecoveryConfig, TEXT("LoginRecovery"));
+	const float DelayTime = AuthEOSGSLoginRecoveryConfig.Interval + AuthEOSGSLoginRecoveryConfig.Jitter * FMath::RandRange(0.f, 1.f);
+
+	if (AccountInfoEOS->RestoreLoginTimer.IsValid())
 	{
-		const FOnlineAccountIdHandle* LocalUserIdPtr = InAsyncOp.Data.Get<FOnlineAccountIdHandle>(TEXT("LocalUserId"));
-		if (!ensure(LocalUserIdPtr))
+		FTSBackgroundableTicker::GetCoreTicker().RemoveTicker(AccountInfoEOS->RestoreLoginTimer);
+	}
+
+	AccountInfoEOS->RestoreLoginTimer = FTSBackgroundableTicker::GetCoreTicker().AddTicker(
+		FTickerDelegate::CreateLambda([this, EpicAccountId = AccountInfoEOS->EpicAccountId](float) -> bool
 		{
-			return MakeFulfilledPromise<TArray<FOnlineAccountIdHandle>>().GetFuture();
-		}
-		return ResolveAccountIds(*LocalUserIdPtr, ProductUserIds);
-	};
+			ConnectLoginRecoveryImplOp(FConnectLoginRecoveryImpl::Params{ EpicAccountId });
+
+			// One-shot timer.
+			return false;
+		}),
+		DelayTime);
 }
 
 FOnlineAccountIdHandle FAuthEOSGS::CreateAccountId(const EOS_ProductUserId ProductUserId)
diff --git a/Engine/Plugins/Online/Experimental/OnlineServicesEOSGS/Source/Public/Online/AuthEOSGS.h b/Engine/Plugins/Online/Experimental/OnlineServicesEOSGS/Source/Public/Online/AuthEOSGS.h
--- a/Engine/Plugins/Online/Experimental/OnlineServicesEOSGS/Source/Public/Online/AuthEOSGS.h
+++ b/Engine/Plugins/Online/Experimental/OnlineServicesEOSGS/Source/Public/Online/AuthEOSGS.h
@@ -4,6 +4,7 @@
 
 #include "CoreMinimal.h"
 #include "Online/AuthCommon.h"
+#include "Online/OnlineServicesEOSGSTypes.h"
 
 #if defined(EOS_PLATFORM_BASE_FILE_NAME)
 #include EOS_PLATFORM_BASE_FILE_NAME
@@ -17,40 +18,38 @@ const int EOS_STRING_BUFFER_LENGTH = 256;
 const int EOS_MAX_TOKEN_SIZE = 4096;
 
 class FOnlineServicesEOSGS;
+struct FAccountInfoEOS;
+class FAccountInfoRegistryEOS;
 
-/* Base class for EOS_Connect login credentials */
-class ONLINESERVICESEOSGS_API FEOSConnectLoginCredentials : public EOS_Connect_Credentials
+struct FAccountInfoEOS final : public FAccountInfo
 {
-public:
-	FEOSConnectLoginCredentials();
-	virtual ~FEOSConnectLoginCredentials() = default;
+	FTSTicker::FDelegateHandle RestoreLoginTimer;
+	EOS_EpicAccountId EpicAccountId = nullptr;
+	EOS_ProductUserId ProductUserId = nullptr;
 };
 
-/* EOS_Connect login credentials using EAS auth */
-class ONLINESERVICESEOSGS_API FEOSConnectLoginCredentialsEAS
-	: public FEOSConnectLoginCredentials
-	, public FNoncopyable
+class ONLINESERVICESEOSGS_API FAccountInfoRegistryEOS final : public FAccountInfoRegistry
 {
 public:
-	FEOSConnectLoginCredentialsEAS(EOS_Auth_Token* InEASToken);
-	virtual ~FEOSConnectLoginCredentialsEAS();
-protected:
-	EOS_Auth_Token* EASToken = nullptr;
-};
+	using Super = FAccountInfoRegistry;
 
-/* EOS_Auth login credentials */
-struct ONLINESERVICESEOSGS_API FEOSAuthCredentials : public EOS_Auth_Credentials
-{
-	FEOSAuthCredentials();
-	FEOSAuthCredentials(EOS_EExternalCredentialType InExternalType, const TArray<uint8>& InToken);
+	virtual ~FAccountInfoRegistryEOS() = default;
+
+	TSharedPtr<FAccountInfoEOS> Find(FPlatformUserId PlatformUserId) const;
+	TSharedPtr<FAccountInfoEOS> Find(FOnlineAccountIdHandle AccountIdHandle) const;
+	TSharedPtr<FAccountInfoEOS> Find(EOS_EpicAccountId EpicAccountId) const;
+	TSharedPtr<FAccountInfoEOS> Find(EOS_ProductUserId ProductUserId) const;
 
-	FEOSAuthCredentials(const FEOSAuthCredentials& Other);
-	FEOSAuthCredentials& operator=(FEOSAuthCredentials& Other);
+	void Register(const TSharedRef<FAccountInfoEOS>& UserAuthData);
+	void Unregister(FOnlineAccountIdHandle AccountId);
 
-	void SetToken(const FCredentialsToken& InToken);
+protected:
+	virtual void DoRegister(const TSharedRef<FAccountInfo>& AccountInfo);
+	virtual void DoUnregister(const TSharedRef<FAccountInfo>& AccountInfo);
 
-	char IdAnsi[EOS_STRING_BUFFER_LENGTH];
-	char TokenAnsi[EOS_MAX_TOKEN_SIZE];
+private:
+	TMap<EOS_EpicAccountId, TSharedRef<FAccountInfoEOS>> AuthDataByEpicAccountId;
+	TMap<EOS_ProductUserId, TSharedRef<FAccountInfoEOS>> AuthDataByProductUserId;
 };
 
 class ONLINESERVICESEOSGS_API FAuthEOSGS : public FAuthCommon
@@ -66,10 +65,13 @@ public:
 	// End IOnlineComponent
 
 	// Begin IAuth
+	virtual void PreShutdown() override;
 	virtual TOnlineAsyncOpHandle<FAuthLogin> Login(FAuthLogin::Params&& Params) override;
 	virtual TOnlineAsyncOpHandle<FAuthLogout> Logout(FAuthLogout::Params&& Params) override;
-	virtual TOnlineResult<FAuthGetAccountByPlatformUserId> GetAccountByPlatformUserId(FAuthGetAccountByPlatformUserId::Params&& Params) override;
-	virtual TOnlineResult<FAuthGetAccountByAccountId> GetAccountByAccountId(FAuthGetAccountByAccountId::Params&& Params) override;
+	virtual TOnlineAsyncOpHandle<FAuthQueryVerifiedAuthTicket> QueryVerifiedAuthTicket(FAuthQueryVerifiedAuthTicket::Params&& Params) override;
+	virtual TOnlineAsyncOpHandle<FAuthCancelVerifiedAuthTicket> CancelVerifiedAuthTicket(FAuthCancelVerifiedAuthTicket::Params&& Params) override;
+	virtual TOnlineAsyncOpHandle<FAuthBeginVerifiedAuthSession> BeginVerifiedAuthSession(FAuthBeginVerifiedAuthSession::Params&& Params) override;
+	virtual TOnlineAsyncOpHandle<FAuthEndVerifiedAuthSession> EndVerifiedAuthSession(FAuthEndVerifiedAuthSession::Params&& Params) override;
 	// End IAuth
 
 	// Begin FAuthEOSGS
@@ -79,23 +81,177 @@ public:
 	virtual TFunction<TFuture<TArray<FOnlineAccountIdHandle>>(FOnlineAsyncOp& InAsyncOp, const TArray<EOS_ProductUserId>& ProductUserIds)> ResolveProductIdsFn();
 	// End FAuthEOSGS
 
-	bool IsLoggedIn(const FOnlineAccountIdHandle& AccountId) const;
+protected:
+	// internal operations.
+
+	struct FLoginEASImpl
+	{
+		static constexpr TCHAR Name[] = TEXT("LoginEASImpl");
+
+		struct Params
+		{
+			FPlatformUserId PlatformUserId = PLATFORMUSERID_NONE;
+			FName CredentialsType;
+			FString CredentialsId;
+			TVariant<FString, FExternalAuthToken> CredentialsToken;
+			TArray<FString> Scopes;
+		};
+
+		struct Result
+		{
+			EOS_EpicAccountId EpicAccountId = nullptr;
+		};
+	};
+
+	TFuture<TDefaultErrorResult<FLoginEASImpl>> LoginEASImpl(const FLoginEASImpl::Params& LoginParams);
+
+	struct FLogoutEASImpl
+	{
+		static constexpr TCHAR Name[] = TEXT("LogoutEASImpl");
+
+		struct Params
+		{
+			EOS_EpicAccountId EpicAccountId = nullptr;
+		};
+
+		struct Result
+		{
+		};
+	};
+
+	TFuture<TDefaultErrorResult<FLogoutEASImpl>> LogoutEASImpl(const FLogoutEASImpl::Params& LogoutParams);
+
+	struct FGetExternalAuthTokenImpl
+	{
+		static constexpr TCHAR Name[] = TEXT("GetExternalAuthTokenImpl");
+
+		struct Params
+		{
+			EOS_EpicAccountId EpicAccountId = nullptr;
+		};
+
+		struct Result
+		{
+			FExternalAuthToken Token;
+		};
+	};
+
+	TDefaultErrorResult<FGetExternalAuthTokenImpl> GetExternalAuthTokenImpl(const FGetExternalAuthTokenImpl::Params& Params);
+
+	struct FLoginConnectImpl
+	{
+		static constexpr TCHAR Name[] = TEXT("LoginConnectImpl");
+
+		struct Params
+		{
+			FPlatformUserId PlatformUserId = PLATFORMUSERID_NONE;
+			FExternalAuthToken ExternalAuthToken;
+		};
+
+		struct Result
+		{
+			EOS_ProductUserId ProductUserId = nullptr;
+		};
+	};
+
+	TFuture<TDefaultErrorResult<FLoginConnectImpl>> LoginConnectImpl(const FLoginConnectImpl::Params& LoginParams);
+
+	struct FConnectLoginRecoveryImpl
+	{
+		static constexpr TCHAR Name[] = TEXT("ConnectLoginRecovery");
+
+		struct Params
+		{
+			/** The Epic Account ID of the local user whose connect login should be recovered. */
+			EOS_EpicAccountId LocalUserId = nullptr;
+		};
+
+		struct Result
+		{
+		};
+	};
+
+	TOnlineAsyncOpHandle<FConnectLoginRecoveryImpl> ConnectLoginRecoveryImplOp(FConnectLoginRecoveryImpl::Params&& Params);
+
+	struct FHandleConnectLoginStatusChangedImpl
+	{
+		static constexpr TCHAR Name[] = TEXT("HandleConnectLoginStatusChangedImpl");
+
+		struct Params
+		{
+			/** The Product User ID of the local player whose status has changed. */
+			EOS_ProductUserId LocalUserId = nullptr;
+			/** The status prior to the change. */
+			EOS_ELoginStatus PreviousStatus = EOS_ELoginStatus::EOS_LS_NotLoggedIn;
+			/** The status at the time of the notification. */
+			EOS_ELoginStatus CurrentStatus = EOS_ELoginStatus::EOS_LS_NotLoggedIn;
+		};
+
+		struct Result
+		{
+		};
+	};
+
+	TOnlineAsyncOpHandle<FHandleConnectLoginStatusChangedImpl> HandleConnectLoginStatusChangedImplOp(FHandleConnectLoginStatusChangedImpl::Params&& Params);
+
+	struct FHandleConnectAuthNotifyExpirationImpl
+	{
+		static constexpr TCHAR Name[] = TEXT("HandleConnectAuthNotifyExpirationImpl");
+
+		struct Params
+		{
+			/** The Product User ID of the local player whose status has changed. */
+			EOS_ProductUserId LocalUserId = nullptr;
+		};
+
+		struct Result
+		{
+		};
+	};
+
+	TOnlineAsyncOpHandle<FHandleConnectAuthNotifyExpirationImpl> HandleConnectAuthNotifyExpirationImplOp(FHandleConnectAuthNotifyExpirationImpl::Params&& Params);
+
+	struct FHandleEASLoginStatusChangedImpl
+	{
+		static constexpr TCHAR Name[] = TEXT("HandleEASLoginStatusChangedImpl");
+
+		struct Params
+		{
+			/** The Epic Account ID of the local user whose status has changed */
+			EOS_EpicAccountId LocalUserId = nullptr;
+			/** The status prior to the change */
+			EOS_ELoginStatus PrevStatus = EOS_ELoginStatus::EOS_LS_NotLoggedIn;
+			/** The status at the time of the notification */
+			EOS_ELoginStatus CurrentStatus = EOS_ELoginStatus::EOS_LS_NotLoggedIn;
+		};
+
+		struct Result
+		{
+		};
+	};
+
+	TOnlineAsyncOpHandle<FHandleEASLoginStatusChangedImpl> HandleEASLoginStatusChangedImplOp(FHandleEASLoginStatusChangedImpl::Params&& Params);
 
 protected:
-	TOnlineChainableAsyncOp<FAuthLogin, TSharedPtr<FEOSConnectLoginCredentials>> LoginEAS(TOnlineAsyncOp<FAuthLogin>& InAsyncOp);
-	void ProcessSuccessfulLogin(TOnlineAsyncOp<FAuthLogin>& InAsyncOp, const EOS_ProductUserId ProductUserId);
+	// Service event handling.
+	void RegisterHandlers();
+	void UnregisterHandlers();
+	void OnConnectLoginStatusChanged(const EOS_Connect_LoginStatusChangedCallbackInfo* Data);
+	void OnConnectAuthNotifyExpiration(const EOS_Connect_AuthExpirationCallbackInfo* Data);
+	void OnEASLoginStatusChanged(const EOS_Auth_LoginStatusChangedCallbackInfo* Data);
 
-	TResult<FOnlineAccountIdHandle, FOnlineError> GetAccountIdByPlatformUserId(FPlatformUserId PlatformUserId) const;
+protected:
+	virtual const FAccountInfoRegistry& GetAccountInfoRegistry() const override;
 
-	static FOnlineAccountIdHandle CreateAccountId(const EOS_ProductUserId ProductUserId);
+	void InitializeConnectLoginRecoveryTimer(const TSharedRef<FAccountInfoEOS>& UserAuthData);
 
-	using FAccountInfoEOS = FAccountInfo;
-	TMap<FOnlineAccountIdHandle, TSharedRef<FAccountInfoEOS>> AccountInfos;
+	static FOnlineAccountIdHandle CreateAccountId(const EOS_ProductUserId ProductUserId);
 
 	EOS_HAuth AuthHandle = nullptr;
 	EOS_HConnect ConnectHandle = nullptr;
+	EOSEventRegistrationPtr OnConnectLoginStatusChangedEOSEventRegistration;
+	EOSEventRegistrationPtr OnConnectAuthNotifyExpirationEOSEventRegistration;
+	FAccountInfoRegistryEOS AccountInfoRegistryEOS;
 };
 
-ELoginStatus ONLINESERVICESEOSGS_API ToELoginStatus(EOS_ELoginStatus InStatus);
-
 /* UE::Online */ }
diff --git a/Engine/Plugins/Online/Experimental/OnlineServicesEOSGS/Source/Public/Online/OnlineServicesEOSGSTypes.h b/Engine/Plugins/Online/Experimental/OnlineServicesEOSGS/Source/Public/Online/OnlineServicesEOSGSTypes.h
--- a/Engine/Plugins/Online/Experimental/OnlineServicesEOSGS/Source/Public/Online/OnlineServicesEOSGSTypes.h
+++ b/Engine/Plugins/Online/Experimental/OnlineServicesEOSGS/Source/Public/Online/OnlineServicesEOSGSTypes.h
@@ -90,6 +90,8 @@ struct FEOSAsyncTraits<void(*)(TEOSHandle, const TEOSParameters*, void*, const T
  * bound, and we can safely consume the CallbackInfo. Note that the CallbackInfo is only valid for the
  * duration of the callback, so it is not safe to bind a continuation _after_ calling the EOS method.
  */
+// Warning: This signature will crash if the async method returns immediately.
+// The callback object is deleted inside the callback before EOS_Async returns the future back to the user.
 template<typename TEOSResult, typename TEOSHandle, typename TEOSParameters, typename TEOSFn>
 decltype(auto) EOS_Async(TEOSFn EOSFn, TEOSHandle EOSHandle, TEOSParameters Parameters, TPromise<const TEOSResult*>&& Promise)
 {
@@ -153,7 +155,7 @@ EOSEventRegistrationPtr EOS_RegisterComponentEventHandler(
 	EOSNotfyUnregisterFunction NotfyUnregisterFunction,
 	ComponentHandlerFunction HandlerFunction);
 
-namespace detail {
+namespace Private {
 
 template<typename Function> struct TEOSCallbackTraitsBase;
 
@@ -206,14 +208,15 @@ public:
 		, NotfyUnregisterFunction(NotfyUnregisterFunction)
 		, HandlerFunction(HandlerFunction)
 	{
+		using EventDataType = typename TEOSCallbackTraits<ComponentHandlerFunction>::EventDataType;
 		typename TEOSNotifyRegisterTraits<EOSNotfyRegisterFunction>::OptionsType Options = { };
 		Options.ApiVersion = ApiVersion;
 		NotificationId = NotfyRegisterFunction(ClientHandle, &Options, this,
-		[](const typename TEOSCallbackTraits<ComponentHandlerFunction>::EventDataType* Data)
+		static_cast<void(EOS_CALL*)(const EventDataType* Data)>([](const EventDataType* Data)
 		{
 			ThisClass* This = reinterpret_cast<ThisClass*>(Data->ClientData);
 			(This->HandlerClass->*This->HandlerFunction)(Data);
-		});
+		}));
 	}
 
 	virtual ~EOSEventRegistrationImpl()
@@ -235,7 +238,7 @@ private:
 	ComponentHandlerFunction HandlerFunction;
 };
 
-/* detail */ }
+/* Private */ }
 
 template <
 	typename ComponentHandlerClass,
@@ -252,7 +255,7 @@ EOSEventRegistrationPtr EOS_RegisterComponentEventHandler(
 	ComponentHandlerFunction HandlerFunction)
 {
 	return MakeUnique<
-		detail::EOSEventRegistrationImpl<
+		Private::EOSEventRegistrationImpl<
 			ComponentHandlerClass,
 			EOSHandle,
 			EOSNotfyRegisterFunction,
