- Fix missing null terminator on converted strings.

diff --git a/Engine/Plugins/Online/Experimental/OnlineServicesEOSGS/Source/Private/Online/AchievementsEOSGS.cpp b/Engine/Plugins/Online/Experimental/OnlineServicesEOSGS/Source/Private/Online/AchievementsEOSGS.cpp
--- a/Engine/Plugins/Online/Experimental/OnlineServicesEOSGS/Source/Private/Online/AchievementsEOSGS.cpp
+++ b/Engine/Plugins/Online/Experimental/OnlineServicesEOSGS/Source/Private/Online/AchievementsEOSGS.cpp
@@ -358,7 +358,7 @@ TOnlineAsyncOpHandle<FUnlockAchievements> FAchievementsEOSGS::UnlockAchievements
 		AchievementIdPtrs.Reserve(InAchievementIds.Num());
 		for (const FString& AchievementId : InAchievementIds)
 		{
-			FTCHARToUTF8& Converter = AchievementIdConverters.Emplace_GetRef(GetData(AchievementId), GetNum(AchievementId));
+			FTCHARToUTF8& Converter = AchievementIdConverters.Emplace_GetRef(*AchievementId);
 			AchievementIdPtrs.Emplace(Converter.Get());
 		}
 
