OSSv2 Sessions Interface - Invite and Overlay operations

diff --git a/Engine/Plugins/Online/Experimental/OnlineServices/Source/OnlineServicesCommon/Private/Online/SessionsCommon.cpp b/Engine/Plugins/Online/Experimental/OnlineServices/Source/OnlineServicesCommon/Private/Online/SessionsCommon.cpp
--- a/Engine/Plugins/Online/Experimental/OnlineServices/Source/OnlineServicesCommon/Private/Online/SessionsCommon.cpp
+++ b/Engine/Plugins/Online/Experimental/OnlineServices/Source/OnlineServicesCommon/Private/Online/SessionsCommon.cpp
@@ -32,7 +32,6 @@ namespace UE::Online {
 		RegisterCommand(&FSessionsCommon::StartMatchmaking);
 		RegisterCommand(&FSessionsCommon::JoinSession);
 		RegisterCommand(&FSessionsCommon::SendSessionInvite);
-		RegisterCommand(&FSessionsCommon::QuerySessionInvites);
 		RegisterCommand(&FSessionsCommon::RejectSessionInvite);
 		RegisterCommand(&FSessionsCommon::RegisterPlayers);
 		RegisterCommand(&FSessionsCommon::UnregisterPlayers);
@@ -123,13 +122,6 @@ namespace UE::Online {
 		return Operation->GetHandle();
 	}
 
-	TOnlineAsyncOpHandle<FQuerySessionInvites> FSessionsCommon::QuerySessionInvites(FQuerySessionInvites::Params&& Params)
-	{
-		TOnlineAsyncOpRef<FQuerySessionInvites> Operation = GetOp<FQuerySessionInvites>(MoveTemp(Params));
-		Operation->SetError(Errors::NotImplemented());
-		return Operation->GetHandle();
-	}
-
 	TOnlineAsyncOpHandle<FRejectSessionInvite> FSessionsCommon::RejectSessionInvite(FRejectSessionInvite::Params&& Params)
 	{
 		TOnlineAsyncOpRef<FRejectSessionInvite> Operation = GetOp<FRejectSessionInvite>(MoveTemp(Params));
@@ -269,14 +261,14 @@ namespace UE::Online {
 		return SessionEvents.OnSessionUpdated;
 	}
 
-	TOnlineEvent<void(const FInviteReceived&)> FSessionsCommon::OnInviteReceived()
+	TOnlineEvent<void(const FSessionInviteReceived&)> FSessionsCommon::OnSessionInviteReceived()
 	{
-		return SessionEvents.OnInviteReceived;
+		return SessionEvents.OnSessionInviteReceived;
 	}
 
-	TOnlineEvent<void(const FInviteAccepted&)> FSessionsCommon::OnInviteAccepted()
+	TOnlineEvent<void(const FUISessionJoinRequested&)> FSessionsCommon::OnUISessionJoinRequested()
 	{
-		return SessionEvents.OnInviteAccepted;
+		return SessionEvents.OnUISessionJoinRequested;
 	}
 
 	FOnlineError FSessionsCommon::CheckCreateSessionParams(const FCreateSession::Params& Params)
diff --git a/Engine/Plugins/Online/Experimental/OnlineServices/Source/OnlineServicesCommon/Public/Online/SessionsCommon.h b/Engine/Plugins/Online/Experimental/OnlineServices/Source/OnlineServicesCommon/Public/Online/SessionsCommon.h
--- a/Engine/Plugins/Online/Experimental/OnlineServices/Source/OnlineServicesCommon/Public/Online/SessionsCommon.h
+++ b/Engine/Plugins/Online/Experimental/OnlineServices/Source/OnlineServicesCommon/Public/Online/SessionsCommon.h
@@ -77,7 +77,6 @@ public:
 	virtual TOnlineAsyncOpHandle<FStartMatchmaking> StartMatchmaking(FStartMatchmaking::Params&& Params) override;
 	virtual TOnlineAsyncOpHandle<FJoinSession> JoinSession(FJoinSession::Params&& Params) override;
 	virtual TOnlineAsyncOpHandle<FSendSessionInvite> SendSessionInvite(FSendSessionInvite::Params&& Params) override;
-	virtual TOnlineAsyncOpHandle<FQuerySessionInvites> QuerySessionInvites(FQuerySessionInvites::Params&& Params) override;
 	virtual TOnlineAsyncOpHandle<FRejectSessionInvite> RejectSessionInvite(FRejectSessionInvite::Params&& Params) override;
 	virtual TOnlineAsyncOpHandle<FRegisterPlayers> RegisterPlayers(FRegisterPlayers::Params&& Params) override;
 	virtual TOnlineAsyncOpHandle<FUnregisterPlayers> UnregisterPlayers(FUnregisterPlayers::Params&& Params) override;
@@ -85,8 +84,8 @@ public:
 	virtual TOnlineEvent<void(const FSessionJoined&)> OnSessionJoined() override;
 	virtual TOnlineEvent<void(const FSessionLeft&)> OnSessionLeft() override;
 	virtual TOnlineEvent<void(const FSessionUpdated&)> OnSessionUpdated() override;
-	virtual TOnlineEvent<void(const FInviteReceived&)> OnInviteReceived() override;
-	virtual TOnlineEvent<void(const FInviteAccepted&)> OnInviteAccepted() override;
+	virtual TOnlineEvent<void(const FSessionInviteReceived&)> OnSessionInviteReceived() override;
+	virtual TOnlineEvent<void(const FUISessionJoinRequested&)> OnUISessionJoinRequested() override;
 
 protected:
 	FOnlineError CheckCreateSessionParams(const FCreateSession::Params& Params);
@@ -109,8 +108,8 @@ protected:
 		TOnlineEventCallable<void(const FSessionJoined&)> OnSessionJoined;
 		TOnlineEventCallable<void(const FSessionLeft&)> OnSessionLeft;
 		TOnlineEventCallable<void(const FSessionUpdated&)> OnSessionUpdated;
-		TOnlineEventCallable<void(const FInviteReceived&)> OnInviteReceived;
-		TOnlineEventCallable<void(const FInviteAccepted&)> OnInviteAccepted;
+		TOnlineEventCallable<void(const FSessionInviteReceived&)> OnSessionInviteReceived;
+		TOnlineEventCallable<void(const FUISessionJoinRequested&)> OnUISessionJoinRequested;
 	} SessionEvents;
 
 	TMap<FName, TSharedRef<FSession>> SessionsByName;
diff --git a/Engine/Plugins/Online/Experimental/OnlineServices/Source/OnlineServicesInterface/Private/Online/Sessions.cpp b/Engine/Plugins/Online/Experimental/OnlineServices/Source/OnlineServicesInterface/Private/Online/Sessions.cpp
--- a/Engine/Plugins/Online/Experimental/OnlineServices/Source/OnlineServicesInterface/Private/Online/Sessions.cpp
+++ b/Engine/Plugins/Online/Experimental/OnlineServices/Source/OnlineServicesInterface/Private/Online/Sessions.cpp
@@ -324,4 +324,31 @@ FSession::FSession(const FSession& InSession)
 
 }
 
+const TCHAR* LexToString(EUISessionJoinRequestedSource UISessionJoinRequestedSource)
+{
+	switch (UISessionJoinRequestedSource)
+	{
+	case EUISessionJoinRequestedSource::FromInvitation:	return TEXT("FromInvitation");
+	default:											checkNoEntry(); // Intentional fallthrough
+	case EUISessionJoinRequestedSource::Unspecified:	return TEXT("Unspecified");
+	}
+}
+
+void LexFromString(EUISessionJoinRequestedSource& OutUISessionJoinRequestedSource, const TCHAR* InStr)
+{
+	if (FCString::Stricmp(InStr, TEXT("FromInvitation")) == 0)
+	{
+		OutUISessionJoinRequestedSource = EUISessionJoinRequestedSource::FromInvitation;
+	}
+	else if (FCString::Stricmp(InStr, TEXT("Unspecified")) == 0)
+	{
+		OutUISessionJoinRequestedSource = EUISessionJoinRequestedSource::Unspecified;
+	}
+	else
+	{
+		checkNoEntry();
+		OutUISessionJoinRequestedSource = EUISessionJoinRequestedSource::Unspecified;
+	}
+}
+
 /* UE::Online */ }
\ No newline at end of file
diff --git a/Engine/Plugins/Online/Experimental/OnlineServices/Source/OnlineServicesInterface/Public/Online/Sessions.h b/Engine/Plugins/Online/Experimental/OnlineServices/Source/OnlineServicesInterface/Public/Online/Sessions.h
--- a/Engine/Plugins/Online/Experimental/OnlineServices/Source/OnlineServicesInterface/Public/Online/Sessions.h
+++ b/Engine/Plugins/Online/Experimental/OnlineServices/Source/OnlineServicesInterface/Public/Online/Sessions.h
@@ -222,14 +222,19 @@ struct ONLINESERVICESINTERFACE_API FSession
 
 struct FSessionInvite
 {
-	/* The user which sent the invite*/
-	FOnlineAccountIdHandle HostUserId;
+	/* The user which the invite got sent to */
+	FOnlineAccountIdHandle RecipientId;
 
-	/* The invite id, needed for rejecting the invite */
-	FName InviteId;
+	/* The user which sent the invite */
+	FOnlineAccountIdHandle SenderId;
 
-	/* The session the invite is for */
-	TSharedRef<FSession> Session;
+	/* The invite id, needed for retrieving session information and rejecting the invite */
+	FString InviteId;
+
+	/* Pointer to the session information */
+	TSharedRef<const FSession> Session;
+
+	// TODO: Default constructor will be deleted after we cache invites
 };
 
 struct FGetAllSessions
@@ -428,7 +433,7 @@ struct FJoinSession
 		FName SessionName;
 
 		/* A reference to the session to be joined. */
-		TSharedRef<const FSession> Session;
+		TSharedRef<const FSession> Session; // TODO: FOnlineSessionIdHandle to be used after we cache search results
 
 		/* Information for all local users who will join the session (includes the session creator)*/
 		FSessionMembersMap LocalUsers;
@@ -464,23 +469,6 @@ struct FSendSessionInvite
 	};
 };
 
-struct FQuerySessionInvites
-{
-	static constexpr TCHAR Name[] = TEXT("QuerySessionInvites");
-
-	struct Params
-	{
-		/* The local user agent which started the query*/
-		FOnlineAccountIdHandle LocalUserId;
-	};
-
-	struct Result
-	{
-		/* Array of session invite objects built from the information returned in the query */
-		TArray<FSessionInvite> SessionInvites;
-	};
-};
-
 struct FRejectSessionInvite
 {
 	static constexpr TCHAR Name[] = TEXT("RejectSessionInvite");
@@ -491,7 +479,7 @@ struct FRejectSessionInvite
 		FOnlineAccountIdHandle LocalUserId;
 
 		/* The invite to be rejected */
-		FSessionInvite SessionInvite;
+		TSharedRef<const FSessionInvite> SessionInvite; // TODO: FOnlineSessionInviteIdHandle to be used after we cache invites
 	};
 
 	struct Result
@@ -550,6 +538,8 @@ struct FSessionJoined
 
 	/* A shared reference to the session joined. */
 	TSharedRef<const FSession> Session;
+
+	FSessionJoined() = delete; // cannot default construct due to TSharedRef
 };
 
 struct FSessionLeft
@@ -565,24 +555,44 @@ struct FSessionUpdated
 
 	/* Updated session settings */
 	FSessionSettingsUpdate UpdatedSettings;
+
+	FSessionUpdated() = delete; // cannot default construct due to TSharedRef
 };
 
-struct FInviteReceived
+struct FSessionInviteReceived
 {
-	/* The local user which accepted the invite */
+	/* The local user which received the invite */
 	FOnlineAccountIdHandle LocalUserId;
 
-	/* Reference to the invite information */
-	TSharedRef<FSessionInvite> SessionInvite;
+	/** The session invite the local user was sent, or the online error if there was a failure retrieving the session for it*/
+	TSharedRef<const FSessionInvite> SessionInvite;
+
+	FSessionInviteReceived() = delete; // cannot default construct due to TSharedRef
+};
+
+/** Session join requested source */
+enum class EUISessionJoinRequestedSource : uint8
+{
+	/** Unspecified by the online service */
+	Unspecified,
+	/** From an invitation */
+	FromInvitation,
 };
+ONLINESERVICESINTERFACE_API const TCHAR* LexToString(EUISessionJoinRequestedSource UISessionJoinRequestedSource);
+ONLINESERVICESINTERFACE_API void LexFromString(EUISessionJoinRequestedSource& OutUISessionJoinRequestedSource, const TCHAR* InStr);
 
-struct FInviteAccepted
+struct FUISessionJoinRequested
 {
-	/* The local user which accepted the invite */
+	/** The local user associated with the join request. */
 	FOnlineAccountIdHandle LocalUserId;
 
-	/* Reference to the invite information */
-	TSharedRef<FSessionInvite> SessionInvite;
+	/** The session the local user requested to join, or the online error if there was a failure retrieving it */
+	TResult<TSharedRef<const FSession>, FOnlineError> Result;
+
+	/** Join request source */
+	EUISessionJoinRequestedSource JoinRequestedSource = EUISessionJoinRequestedSource::Unspecified;
+
+	FUISessionJoinRequested() = delete; // cannot default construct due to TSharedRef
 };
 
 class ISessions
@@ -668,14 +678,6 @@ public:
 	 */
 	virtual TOnlineAsyncOpHandle<FSendSessionInvite> SendSessionInvite(FSendSessionInvite::Params&& Params) = 0;
 
-	/**
-	 * Queries all session invites for a given user.
-	 *
-	 * @param Parameters for the QuerySessionInvites call
-	 * @return
-	 */
-	virtual TOnlineAsyncOpHandle<FQuerySessionInvites> QuerySessionInvites(FQuerySessionInvites::Params&& Params) = 0;
-
 	/**
 	 * Rejects a given session invite for a user.
 	 *
@@ -732,15 +734,15 @@ public:
 	 *
 	 * @return
 	 */
-	virtual TOnlineEvent<void(const FInviteReceived&)> OnInviteReceived() = 0;
+	virtual TOnlineEvent<void(const FSessionInviteReceived&)> OnSessionInviteReceived() = 0;
 
 	/**
-	 * Get the event that is triggered when a session invite is accepted.
-	 * This event will trigger as a result of accepting a platform session invite.
+	 * Get the event that is triggered when a session is joined via UI.
+	 * This event will trigger as a result of joining a session via the platform UI.
 	 *
 	 * @return
 	 */
-	virtual TOnlineEvent<void(const FInviteAccepted&)> OnInviteAccepted() = 0;
+	virtual TOnlineEvent<void(const FUISessionJoinRequested&)> OnUISessionJoinRequested() = 0;
 };
 
 namespace Meta {
@@ -816,7 +818,7 @@ BEGIN_ONLINE_STRUCT_META(FSession)
 END_ONLINE_STRUCT_META()
 
 BEGIN_ONLINE_STRUCT_META(FSessionInvite)
-	ONLINE_STRUCT_FIELD(FSessionInvite, HostUserId),
+	ONLINE_STRUCT_FIELD(FSessionInvite, SenderId),
 	ONLINE_STRUCT_FIELD(FSessionInvite, InviteId),
 	ONLINE_STRUCT_FIELD(FSessionInvite, Session)
 END_ONLINE_STRUCT_META()
@@ -919,14 +921,6 @@ END_ONLINE_STRUCT_META()
 BEGIN_ONLINE_STRUCT_META(FSendSessionInvite::Result)
 END_ONLINE_STRUCT_META()
 
-BEGIN_ONLINE_STRUCT_META(FQuerySessionInvites::Params)
-	ONLINE_STRUCT_FIELD(FQuerySessionInvites::Params, LocalUserId)
-END_ONLINE_STRUCT_META()
-
-BEGIN_ONLINE_STRUCT_META(FQuerySessionInvites::Result)
-	ONLINE_STRUCT_FIELD(FQuerySessionInvites::Result, SessionInvites)
-END_ONLINE_STRUCT_META()
-
 BEGIN_ONLINE_STRUCT_META(FRejectSessionInvite::Params)
 	ONLINE_STRUCT_FIELD(FRejectSessionInvite::Params, LocalUserId),
 	ONLINE_STRUCT_FIELD(FRejectSessionInvite::Params, SessionInvite)
diff --git a/Engine/Plugins/Online/Experimental/OnlineServicesEOSGS/Source/Private/Online/SessionsEOSGS.cpp b/Engine/Plugins/Online/Experimental/OnlineServicesEOSGS/Source/Private/Online/SessionsEOSGS.cpp
--- a/Engine/Plugins/Online/Experimental/OnlineServicesEOSGS/Source/Private/Online/SessionsEOSGS.cpp
+++ b/Engine/Plugins/Online/Experimental/OnlineServicesEOSGS/Source/Private/Online/SessionsEOSGS.cpp
@@ -51,7 +51,7 @@ FSessionEOSGS::FSessionEOSGS(const FSession& InSession)
 }
 
 FSessionEOSGS::FSessionEOSGS(const EOS_HSessionDetails& InSessionDetailsHandle)
-	: SessionDetailsHandle(MakeShareable<FSessionDetailsHandleEOSGS>(new FSessionDetailsHandleEOSGS(InSessionDetailsHandle)))
+	: SessionDetailsHandle(MakeShared<FSessionDetailsHandleEOSGS>(InSessionDetailsHandle))
 {
 	EOS_SessionDetails_CopyInfoOptions CopyInfoOptions = { };
 	CopyInfoOptions.ApiVersion = EOS_SESSIONDETAILS_COPYINFO_API_LATEST;
@@ -78,7 +78,7 @@ const FSessionEOSGS& FSessionEOSGS::Cast(const FSession& InSession)
 {
 	check(InSession.SessionId.GetOnlineServicesType() == EOnlineServices::Epic);
 
-	return *static_cast<const FSessionEOSGS*>(&InSession);
+	return static_cast<const FSessionEOSGS&>(InSession);
 }
 
 /** FSessionsEOSGS */
@@ -95,6 +95,121 @@ void FSessionsEOSGS::Initialize()
 
 	SessionsHandle = EOS_Platform_GetSessionsInterface(static_cast<FOnlineServicesEOSGS&>(GetServices()).GetEOSPlatformHandle());
 	check(SessionsHandle);
+
+	RegisterEventHandlers();
+}
+
+void FSessionsEOSGS::Shutdown()
+{
+	Super::Shutdown();
+
+	UnregisterEventHandlers();
+}
+
+void FSessionsEOSGS::RegisterEventHandlers()
+{
+	// Register for session invites received events
+	OnSessionInviteReceivedEventRegistration = EOS_RegisterComponentEventHandler(
+		this,
+		SessionsHandle,
+		EOS_SESSIONS_ADDNOTIFYSESSIONINVITERECEIVED_API_LATEST,
+		&EOS_Sessions_AddNotifySessionInviteReceived,
+		&EOS_Sessions_RemoveNotifySessionInviteReceived,
+		&FSessionsEOSGS::HandleSessionInviteReceived);
+
+	static_assert(EOS_SESSIONS_ADDNOTIFYSESSIONINVITERECEIVED_API_LATEST == 1, "EOS_Sessions_AddNotifySessionInviteReceivedOptions updated, check new fields");
+
+	// Register for session invites accepted events
+	OnSessionInviteAcceptedEventRegistration = EOS_RegisterComponentEventHandler(
+		this,
+		SessionsHandle,
+		EOS_SESSIONS_ADDNOTIFYSESSIONINVITEACCEPTED_API_LATEST,
+		&EOS_Sessions_AddNotifySessionInviteAccepted,
+		&EOS_Sessions_RemoveNotifySessionInviteAccepted,
+		&FSessionsEOSGS::HandleSessionInviteAccepted);
+
+	static_assert(EOS_SESSIONS_ADDNOTIFYSESSIONINVITEACCEPTED_API_LATEST == 1, "EOS_Sessions_AddNotifySessionInviteAcceptedOptions updated, check new fields");
+
+	// Register for join session accepted events
+	OnJoinSessionAcceptedEventRegistration = EOS_RegisterComponentEventHandler(
+		this,
+		SessionsHandle,
+		EOS_SESSIONS_ADDNOTIFYJOINSESSIONACCEPTED_API_LATEST,
+		&EOS_Sessions_AddNotifyJoinSessionAccepted,
+		&EOS_Sessions_RemoveNotifyJoinSessionAccepted,
+		&FSessionsEOSGS::HandleJoinSessionAccepted);
+
+	static_assert(EOS_SESSIONS_ADDNOTIFYJOINSESSIONACCEPTED_API_LATEST == 1, "EOS_Sessions_AddNotifyJoinSessionAcceptedOptions updated, check new fields");
+}
+
+void FSessionsEOSGS::UnregisterEventHandlers()
+{
+	OnSessionInviteReceivedEventRegistration = nullptr;
+	OnSessionInviteAcceptedEventRegistration = nullptr;
+	OnJoinSessionAcceptedEventRegistration = nullptr;
+}
+
+void FSessionsEOSGS::HandleSessionInviteReceived(const EOS_Sessions_SessionInviteReceivedCallbackInfo* Data)
+{
+	FOnlineAccountIdHandle LocalUserId = FindAccountId(Data->LocalUserId);
+	if (LocalUserId.IsValid())
+	{
+		Services.Get<FAuthEOSGS>()->ResolveAccountIds(LocalUserId, { Data->LocalUserId, Data->TargetUserId })
+			.Next([this, WeakThis = AsWeak(), LocalUserId, InviteId = FString(Data->InviteId)](TArray<FOnlineAccountIdHandle> ResolvedAccountIds)
+		{
+			if (const TSharedPtr<ISessions> StrongThis = WeakThis.Pin())
+			{
+				// First and second place in the array will be occupied by the receiver and the sender, respectively, since the same order is kept in the array of resolved ids
+				const FOnlineAccountIdHandle& ReceiverId = ResolvedAccountIds[0];
+				const FOnlineAccountIdHandle& SenderId = ResolvedAccountIds[1];
+
+				TResult<TSharedRef<const FSessionInvite>, FOnlineError> SessionInviteResult = BuildSessionInvite(ReceiverId, SenderId, InviteId);
+
+				if (SessionInviteResult.IsOk())
+				{
+					FSessionInviteReceived Event = { LocalUserId, SessionInviteResult.GetOkValue() };
+
+					SessionEvents.OnSessionInviteReceived.Broadcast(Event);
+				}
+
+				// We won't broadcast the event if there was an error retrieving the session information
+			}
+		});
+	}
+}
+
+void FSessionsEOSGS::HandleSessionInviteAccepted(const EOS_Sessions_SessionInviteAcceptedCallbackInfo* Data)
+{
+	FOnlineAccountIdHandle LocalUserId = FindAccountId(Data->LocalUserId);
+	if (LocalUserId.IsValid())
+	{
+		FUISessionJoinRequested Event{
+			LocalUserId,
+			BuildSessionFromInvite(FString(Data->InviteId)),
+			EUISessionJoinRequestedSource::FromInvitation
+		};
+
+		SessionEvents.OnUISessionJoinRequested.Broadcast(Event);
+
+		// The game can react to the OnUISessionJoinRequested event by starting the JoinSession process
+	}
+}
+
+void FSessionsEOSGS::HandleJoinSessionAccepted(const EOS_Sessions_JoinSessionAcceptedCallbackInfo* Data)
+{
+	FOnlineAccountIdHandle LocalUserId = FindAccountId(Data->LocalUserId);
+	if (LocalUserId.IsValid())
+	{
+		FUISessionJoinRequested Event {
+			LocalUserId,
+			BuildSessionFromUIEvent(Data->UiEventId),
+			EUISessionJoinRequestedSource::Unspecified
+		};
+
+		SessionEvents.OnUISessionJoinRequested.Broadcast(Event);
+
+		// The game can react to the OnUISessionJoinRequested event by starting the JoinSession process
+	}
 }
 
 TOnlineAsyncOpHandle<FCreateSession> FSessionsEOSGS::CreateSession(FCreateSession::Params&& Params)
@@ -837,7 +952,7 @@ TOnlineAsyncOpHandle<FFindSessions> FSessionsEOSGS::FindSessions(FFindSessions::
 				return;
 			}
 
-			CurrentSessionSearchHandleEOSGS = MakeShareable(new FSessionSearchHandleEOSGS(SearchHandle));
+			CurrentSessionSearchHandleEOSGS = MakeShared<FSessionSearchHandleEOSGS>(SearchHandle);
 
 			// We write the search attributes
 			WriteSessionSearchHandle(*CurrentSessionSearchHandleEOSGS, OpParams);
@@ -986,6 +1101,10 @@ TOnlineAsyncOpHandle<FJoinSession> FSessionsEOSGS::JoinSession(FJoinSession::Par
 		if (TSharedRef<FSessionEOSGS>* FoundSession = SessionsByName.Find(OpParams.SessionName))
 		{
 			FoundSession->Get().CurrentState = ESessionState::Valid;
+
+			FSessionJoined Event = { { OpParams.LocalUserId }, *FoundSession };
+
+			SessionEvents.OnSessionJoined.Broadcast(Event);
 		}
 
 		// A successful join allows the client to server travel, after which RegisterPlayers will be called by the engine
@@ -997,4 +1116,148 @@ TOnlineAsyncOpHandle<FJoinSession> FSessionsEOSGS::JoinSession(FJoinSession::Par
 	return Op->GetHandle();
 }
 
+TResult<TSharedRef<const FSessionInvite>, FOnlineError> FSessionsEOSGS::BuildSessionInvite(FOnlineAccountIdHandle RecipientId, FOnlineAccountIdHandle SenderId, const FString& InInviteId) const
+{
+	TResult<TSharedRef<const FSession>, FOnlineError> SessionInviteResult = BuildSessionFromInvite(InInviteId);
+
+	if (SessionInviteResult.IsOk())
+	{
+		FSessionInvite SessionInvite {
+			RecipientId,
+			SenderId,
+			InInviteId,
+			SessionInviteResult.GetOkValue()
+		};
+
+		return TResult<TSharedRef<const FSessionInvite>, FOnlineError>(MakeShared<FSessionInvite>(MoveTemp(SessionInvite)));
+	}
+	else
+	{
+		return TResult<TSharedRef<const FSessionInvite>, FOnlineError>(SessionInviteResult.GetErrorValue());
+	}
+}
+
+TResult<TSharedRef<const FSession>, FOnlineError> FSessionsEOSGS::BuildSessionFromInvite(const FString& InInviteId) const
+{
+	EOS_Sessions_CopySessionHandleByInviteIdOptions CopySessionHandleByInviteIdOptions = { };
+	CopySessionHandleByInviteIdOptions.ApiVersion = EOS_SESSIONS_COPYSESSIONHANDLEBYINVITEID_API_LATEST;
+	static_assert(EOS_SESSIONS_COPYSESSIONHANDLEBYINVITEID_API_LATEST == 1, "EOS_Sessions_CopySessionHandleByInviteIdOptions updated, check new fields");
+
+	FTCHARToUTF8 InviteIdUtf8(InInviteId);
+	CopySessionHandleByInviteIdOptions.InviteId = InviteIdUtf8.Get();
+
+	EOS_HSessionDetails SessionDetailsHandle;
+	EOS_EResult CopySessionHandleByInviteIdResult = EOS_Sessions_CopySessionHandleByInviteId(SessionsHandle, &CopySessionHandleByInviteIdOptions, &SessionDetailsHandle);
+	if (CopySessionHandleByInviteIdResult == EOS_EResult::EOS_Success)
+	{
+		return TResult<TSharedRef<const FSession>, FOnlineError>(MakeShared<FSessionEOSGS>(SessionDetailsHandle));
+	}
+	else
+	{
+		UE_LOG(LogTemp, Error, TEXT("[FSessionsEOSGS::BuildSessionFromInvite] EOS_Sessions_CopySessionHandleByInviteId failed with result [%s]"), *LexToString(CopySessionHandleByInviteIdResult));
+
+		return TResult<TSharedRef<const FSession>, FOnlineError>(FromEOSError(CopySessionHandleByInviteIdResult));
+	}
+}
+
+TResult<TSharedRef<const FSession>, FOnlineError> FSessionsEOSGS::BuildSessionFromUIEvent(const EOS_UI_EventId& UIEventId) const
+{
+	EOS_Sessions_CopySessionHandleByUiEventIdOptions CopySessionHandleByUiEventIdOptions = { };
+	CopySessionHandleByUiEventIdOptions.ApiVersion = EOS_SESSIONS_COPYSESSIONHANDLEBYUIEVENTID_API_LATEST;
+	static_assert(EOS_SESSIONS_COPYSESSIONHANDLEBYUIEVENTID_API_LATEST == 1, "EOS_Sessions_CopySessionHandleByUiEventIdOptions updated, check new fields");
+
+	CopySessionHandleByUiEventIdOptions.UiEventId = UIEventId;
+
+	EOS_HSessionDetails SessionDetailsHandle;
+	EOS_EResult CopySessionHandleByUiEventIdResult = EOS_Sessions_CopySessionHandleByUiEventId(SessionsHandle, &CopySessionHandleByUiEventIdOptions, &SessionDetailsHandle);
+	if (CopySessionHandleByUiEventIdResult == EOS_EResult::EOS_Success)
+	{
+		return TResult<TSharedRef<const FSession>, FOnlineError>(MakeShared<FSessionEOSGS>(SessionDetailsHandle));
+	}
+	else
+	{
+		UE_LOG(LogTemp, Error, TEXT("[FSessionsEOSGS::BuildSessionFromUIEvent] EOS_Sessions_CopySessionHandleByUiEventId failed with result [%s]"), *LexToString(CopySessionHandleByUiEventIdResult));
+
+		return TResult<TSharedRef<const FSession>, FOnlineError>(FromEOSError(CopySessionHandleByUiEventIdResult));
+	}
+}
+
+TOnlineAsyncOpHandle<FSendSessionInvite> FSessionsEOSGS::SendSessionInvite(FSendSessionInvite::Params&& Params)
+{
+	TOnlineAsyncOpRef<FSendSessionInvite> Op = GetOp<FSendSessionInvite>(MoveTemp(Params));
+
+	// TODO: Param checks
+	Op->Then([this](TOnlineAsyncOp<FSendSessionInvite>& Op, TPromise<const EOS_Sessions_SendInviteCallbackInfo*>&& Promise)
+	{
+		const FSendSessionInvite::Params& OpParams = Op.GetParams();
+
+		// TODO: State checks
+
+		EOS_Sessions_SendInviteOptions SendInviteOptions = { };
+		SendInviteOptions.ApiVersion = EOS_SESSIONS_SENDINVITE_API_LATEST;
+		static_assert(EOS_SESSIONS_SENDINVITE_API_LATEST == 1, "EOS_Sessions_SendInviteOptions updated, check new fields");
+
+		SendInviteOptions.LocalUserId = GetProductUserIdChecked(OpParams.LocalUserId);
+		
+		const FTCHARToUTF8 SessionNameUtf8(OpParams.SessionName.ToString());
+		SendInviteOptions.SessionName = SessionNameUtf8.Get();
+
+		SendInviteOptions.TargetUserId = GetProductUserIdChecked(OpParams.TargetUsers[0]); // TODO: Multiple user invitations using WhenAll like JoinLobbyImpl
+
+		EOS_Async<EOS_Sessions_SendInviteCallbackInfo>(EOS_Sessions_SendInvite, SessionsHandle, SendInviteOptions, MoveTemp(Promise));
+		return;
+	})
+	.Then([this](TOnlineAsyncOp<FSendSessionInvite>& Op, const EOS_Sessions_SendInviteCallbackInfo* Result)
+	{
+		if (Result->ResultCode != EOS_EResult::EOS_Success)
+		{
+			Op.SetError(FromEOSError(Result->ResultCode));
+			return;
+		}
+
+		Op.SetResult(FSendSessionInvite::Result{ });
+	})
+	.Enqueue(GetSerialQueue()); // TODO: Use the parallel queue instead when possible
+
+	return Op->GetHandle();
+}
+
+TOnlineAsyncOpHandle<FRejectSessionInvite> FSessionsEOSGS::RejectSessionInvite(FRejectSessionInvite::Params&& Params)
+{
+	TOnlineAsyncOpRef<FRejectSessionInvite> Op = GetOp<FRejectSessionInvite>(MoveTemp(Params));
+
+	// TODO: Param checks
+	Op->Then([this](TOnlineAsyncOp<FRejectSessionInvite>& Op, TPromise<const EOS_Sessions_RejectInviteCallbackInfo*>&& Promise)
+	{
+		const FRejectSessionInvite::Params& OpParams = Op.GetParams();
+
+		// TODO: State checks
+
+		EOS_Sessions_RejectInviteOptions RejectInviteOptions = { };
+		RejectInviteOptions.ApiVersion = EOS_SESSIONS_REJECTINVITE_API_LATEST;
+		static_assert(EOS_SESSIONS_REJECTINVITE_API_LATEST == 1, "EOS_Sessions_RejectInviteOptions updated, check new fields");
+
+		RejectInviteOptions.LocalUserId = GetProductUserIdChecked(OpParams.LocalUserId);
+
+		const FTCHARToUTF8 InviteIdUtf8(OpParams.SessionInvite->InviteId);
+		RejectInviteOptions.InviteId = InviteIdUtf8.Get();
+
+		EOS_Async<EOS_Sessions_RejectInviteCallbackInfo>(EOS_Sessions_RejectInvite, SessionsHandle, RejectInviteOptions, MoveTemp(Promise));
+		return;
+	})
+	.Then([this](TOnlineAsyncOp<FRejectSessionInvite>& Op, const EOS_Sessions_RejectInviteCallbackInfo* Result)
+	{
+		if (Result->ResultCode != EOS_EResult::EOS_Success)
+		{
+			Op.SetError(FromEOSError(Result->ResultCode));
+			return;
+		}
+
+		Op.SetResult(FRejectSessionInvite::Result{ });
+	})
+	.Enqueue(GetSerialQueue()); // TODO: Use the parallel queue instead when possible
+
+	return Op->GetHandle();
+}
+
 /* UE::Online */ }
\ No newline at end of file
diff --git a/Engine/Plugins/Online/Experimental/OnlineServicesEOSGS/Source/Public/Online/SessionsEOSGS.h b/Engine/Plugins/Online/Experimental/OnlineServicesEOSGS/Source/Public/Online/SessionsEOSGS.h
--- a/Engine/Plugins/Online/Experimental/OnlineServicesEOSGS/Source/Public/Online/SessionsEOSGS.h
+++ b/Engine/Plugins/Online/Experimental/OnlineServicesEOSGS/Source/Public/Online/SessionsEOSGS.h
@@ -4,12 +4,14 @@
 
 #include "CoreMinimal.h"
 #include "Online/SessionsCommon.h"
+#include "Online/OnlineServicesEOSGSTypes.h"
 #include "IPAddress.h"
 
 #if defined(EOS_PLATFORM_BASE_FILE_NAME)
 #include EOS_PLATFORM_BASE_FILE_NAME
 #endif
 #include "eos_sessions_types.h"
+#include "eos_ui_types.h"
 
 namespace UE::Online {
 
@@ -32,7 +34,7 @@ struct FSessionSearchHandleEOSGS : FNoncopyable
 	{
 	}
 
-	virtual ~FSessionSearchHandleEOSGS()
+	~FSessionSearchHandleEOSGS()
 	{
 		EOS_SessionSearch_Release(SearchHandle);
 	}
@@ -47,7 +49,7 @@ struct FSessionDetailsHandleEOSGS : FNoncopyable
 	{
 	}
 
-	virtual ~FSessionDetailsHandleEOSGS()
+	~FSessionDetailsHandleEOSGS()
 	{
 		EOS_SessionDetails_Release(SessionDetailsHandle);
 	}
@@ -94,6 +96,7 @@ public:
 
 	// TOnlineComponent
 	void Initialize() override;
+	void Shutdown() override;
 
 	// ISessions
 	virtual TOnlineAsyncOpHandle<FCreateSession> CreateSession(FCreateSession::Params&& Params) override;
@@ -101,8 +104,16 @@ public:
 	virtual TOnlineAsyncOpHandle<FLeaveSession> LeaveSession(FLeaveSession::Params&& Params) override;
 	virtual TOnlineAsyncOpHandle<FFindSessions> FindSessions(FFindSessions::Params&& Params) override;
 	virtual TOnlineAsyncOpHandle<FJoinSession> JoinSession(FJoinSession::Params&& Params) override;
+	virtual TOnlineAsyncOpHandle<FSendSessionInvite> SendSessionInvite(FSendSessionInvite::Params&& Params) override;
+	virtual TOnlineAsyncOpHandle<FRejectSessionInvite> RejectSessionInvite(FRejectSessionInvite::Params&& Params) override;
 
 protected:
+	void RegisterEventHandlers();
+	void UnregisterEventHandlers();
+	void HandleSessionInviteReceived(const EOS_Sessions_SessionInviteReceivedCallbackInfo* Data);
+	void HandleSessionInviteAccepted(const EOS_Sessions_SessionInviteAcceptedCallbackInfo* Data);
+	void HandleJoinSessionAccepted(const EOS_Sessions_JoinSessionAcceptedCallbackInfo* Data);
+
 	void SetPermissionLevel(EOS_HSessionModification& SessionModHandle, const ESessionJoinPolicy& JoinPolicy);
 	void SetBucketId(EOS_HSessionModification& SessionModHandle, const FString& NewBucketId);
 	void SetMaxPlayers(EOS_HSessionModification& SessionModHandle, const uint32& NewMaxPlayers);
@@ -121,11 +132,30 @@ protected:
 
 	static FOnlineSessionIdHandle CreateSessionId(const FString& SessionId);
 
+	/**
+	 * Builds a session invite from an invite id, using the class' Sessions Handle
+	 */
+	TResult<TSharedRef<const FSessionInvite>, FOnlineError> BuildSessionInvite(FOnlineAccountIdHandle RecipientId, FOnlineAccountIdHandle SenderId, const FString& InInviteId) const;
+
+	/**
+	* Builds a session from an invite id, using the class' Sessions Handle
+	 */
+	TResult<TSharedRef<const FSession>, FOnlineError> BuildSessionFromInvite(const FString& InInviteId) const;
+
+	/**
+	 * Builds a session from a UI event id, using the class' Sessions Handle
+	 */
+	TResult<TSharedRef<const FSession>, FOnlineError> BuildSessionFromUIEvent(const EOS_UI_EventId& UIEventId) const;
+
 protected:
 	FOnlineServicesEOSGS& Services;
 
 	EOS_HSessions SessionsHandle = nullptr;
 
+	EOSEventRegistrationPtr OnSessionInviteReceivedEventRegistration;
+	EOSEventRegistrationPtr OnSessionInviteAcceptedEventRegistration;
+	EOSEventRegistrationPtr OnJoinSessionAcceptedEventRegistration;
+
 	TMap<FName, TSharedRef<FSessionEOSGS>> SessionsByName;
 	TMap<FOnlineSessionIdHandle, TSharedRef<FSessionEOSGS>> SessionsById;
 
diff --git a/Engine/Plugins/Online/Experimental/OnlineServicesNull/Source/Private/Online/SessionsNull.cpp b/Engine/Plugins/Online/Experimental/OnlineServicesNull/Source/Private/Online/SessionsNull.cpp
--- a/Engine/Plugins/Online/Experimental/OnlineServicesNull/Source/Private/Online/SessionsNull.cpp
+++ b/Engine/Plugins/Online/Experimental/OnlineServicesNull/Source/Private/Online/SessionsNull.cpp
@@ -92,7 +92,7 @@ FSessionsNull::FSessionsNull(FOnlineServicesNull& InServices)
 {
 	if (!LANSessionManager.IsValid())
 	{
-		LANSessionManager = MakeShareable(new FLANSession());
+		LANSessionManager = MakeShared<FLANSession>();
 	}
 }
 
@@ -365,17 +365,16 @@ TOnlineAsyncOpHandle<FJoinSession> FSessionsNull::JoinSession(FJoinSession::Para
 			return;
 		}
 
-		TSharedRef<FSessionNull> NewSessionNullRef = MakeShared<FSessionNull>(FSessionNull(*OpParams.Session));
+		TSharedRef<FSessionNull> NewSessionNullRef = MakeShared<FSessionNull>(*OpParams.Session);
 		FSessionSettings NewSessionNullSettings = NewSessionNullRef->SessionSettings;
 
-		FSessionJoined SessionJoinedEvent;
-		SessionJoinedEvent.Session = NewSessionNullRef;
-
 		// TODO: This will move to AddSessionMember
 		TArray<FOnlineAccountIdHandle> LocalUserIds;
 		LocalUserIds.Reserve(OpParams.LocalUsers.Num());
 		OpParams.LocalUsers.GenerateKeyArray(LocalUserIds);
 
+		TArray<FOnlineAccountIdHandle> JoinedUsers;
+
 		for (const FOnlineAccountIdHandle& LocalUserId : LocalUserIds)
 		{
 			bool bReserveSlot = true;
@@ -397,9 +396,11 @@ TOnlineAsyncOpHandle<FJoinSession> FSessionsNull::JoinSession(FJoinSession::Para
 				}
 			}
 
-			SessionJoinedEvent.LocalUserIds.Add(LocalUserId);
+			JoinedUsers.Add(LocalUserId);
 		}
 
+		FSessionJoined SessionJoinedEvent = { MoveTemp(JoinedUsers), NewSessionNullRef };
+
 		NewSessionNullSettings.SessionMembers.Append(OpParams.LocalUsers);
 
 		SessionsByName.Add(OpParams.SessionName, NewSessionNullRef);
