Add metadata for internal Auth and Lobby operations.

diff --git a/Engine/Plugins/Online/EOSShared/Source/EOSShared/Private/EOSShared.cpp b/Engine/Plugins/Online/EOSShared/Source/EOSShared/Private/EOSShared.cpp
--- a/Engine/Plugins/Online/EOSShared/Source/EOSShared/Private/EOSShared.cpp
+++ b/Engine/Plugins/Online/EOSShared/Source/EOSShared/Private/EOSShared.cpp
@@ -62,6 +62,17 @@ const TCHAR* LexToString(const EOS_EFriendsStatus FriendStatus)
 	}
 }
 
+const TCHAR* LexToString(EOS_ELoginStatus LoginStatus)
+{
+	switch (LoginStatus)
+	{
+	default: checkNoEntry(); // Intentional fallthrough
+	case EOS_ELoginStatus::EOS_LS_NotLoggedIn:			return TEXT("EOS_LS_NotLoggedIn");
+	case EOS_ELoginStatus::EOS_LS_UsingLocalProfile:	return TEXT("EOS_LS_UsingLocalProfile");
+	case EOS_ELoginStatus::EOS_LS_LoggedIn:				return TEXT("EOS_LS_LoggedIn");
+	}
+}
+
 bool LexFromString(EOS_EExternalCredentialType& OutEnum, const TCHAR* InString)
 {
 	if (FCString::Stricmp(InString, TEXT("Steam")) == 0)
diff --git a/Engine/Plugins/Online/EOSShared/Source/EOSShared/Public/EOSShared.h b/Engine/Plugins/Online/EOSShared/Source/EOSShared/Public/EOSShared.h
--- a/Engine/Plugins/Online/EOSShared/Source/EOSShared/Public/EOSShared.h
+++ b/Engine/Plugins/Online/EOSShared/Source/EOSShared/Public/EOSShared.h
@@ -41,6 +41,7 @@ inline EOS_ProductUserId EOSProductUserIdFromString(const TCHAR* String)
 EOSSHARED_API FString LexToString(const EOS_EpicAccountId AccountId);
 
 EOSSHARED_API const TCHAR* LexToString(const EOS_EFriendsStatus FriendStatus);
+EOSSHARED_API const TCHAR* LexToString(const EOS_ELoginStatus LoginStatus);
 EOSSHARED_API bool LexFromString(EOS_EAuthScopeFlags& OutEnum, const TCHAR* InString);
 EOSSHARED_API bool LexFromString(EOS_EExternalCredentialType& OutEnum, const TCHAR* InString);
 EOSSHARED_API bool LexFromString(EOS_ELoginCredentialType& OutEnum, const TCHAR* InString);
diff --git a/Engine/Plugins/Online/Experimental/OnlineServices/Source/OnlineServicesInterface/Private/Online/Lobbies.cpp b/Engine/Plugins/Online/Experimental/OnlineServices/Source/OnlineServicesInterface/Private/Online/Lobbies.cpp
--- a/Engine/Plugins/Online/Experimental/OnlineServices/Source/OnlineServicesInterface/Private/Online/Lobbies.cpp
+++ b/Engine/Plugins/Online/Experimental/OnlineServices/Source/OnlineServicesInterface/Private/Online/Lobbies.cpp
@@ -36,6 +36,70 @@ void LexFromString(ELobbyJoinPolicy& OutPolicy, const TCHAR* InStr)
 	}
 }
 
+const TCHAR* LexToString(ELobbyMemberLeaveReason LeaveReason)
+{
+	switch (LeaveReason)
+	{
+	case ELobbyMemberLeaveReason::Left:			return TEXT("Left");
+	case ELobbyMemberLeaveReason::Kicked:		return TEXT("Kicked");
+	case ELobbyMemberLeaveReason::Disconnected:	return TEXT("Disconnected");
+	default:									checkNoEntry(); // Intentional fallthrough
+	case ELobbyMemberLeaveReason::Closed:		return TEXT("Closed");
+	}
+}
+
+void LexFromString(ELobbyMemberLeaveReason& OutLeaveReason, const TCHAR* InStr)
+{
+	if (FCString::Stricmp(InStr, TEXT("Left")) == 0)
+	{
+		OutLeaveReason = ELobbyMemberLeaveReason::Left;
+	}
+	else if (FCString::Stricmp(InStr, TEXT("Kicked")) == 0)
+	{
+		OutLeaveReason = ELobbyMemberLeaveReason::Kicked;
+	}
+	else if (FCString::Stricmp(InStr, TEXT("Disconnected")) == 0)
+	{
+		OutLeaveReason = ELobbyMemberLeaveReason::Disconnected;
+	}
+	else if (FCString::Stricmp(InStr, TEXT("Closed")) == 0)
+	{
+		OutLeaveReason = ELobbyMemberLeaveReason::Closed;
+	}
+	else
+	{
+		checkNoEntry();
+		OutLeaveReason = ELobbyMemberLeaveReason::Closed;
+	}
+}
+
+const TCHAR* LexToString(ELobbyAttributeVisibility Visibility)
+{
+	switch (Visibility)
+	{
+	case ELobbyAttributeVisibility::Public:		return TEXT("Public");
+	default:									checkNoEntry(); // Intentional fallthrough
+	case ELobbyAttributeVisibility::Private:	return TEXT("Private");
+	}
+}
+
+void LexFromString(ELobbyAttributeVisibility& OutVisibility, const TCHAR* InStr)
+{
+	if (FCString::Stricmp(InStr, TEXT("Public")) == 0)
+	{
+		OutVisibility = ELobbyAttributeVisibility::Public;
+	}
+	else if (FCString::Stricmp(InStr, TEXT("Private")) == 0)
+	{
+		OutVisibility = ELobbyAttributeVisibility::Private;
+	}
+	else
+	{
+		checkNoEntry();
+		OutVisibility = ELobbyAttributeVisibility::Private;
+	}
+}
+
 const TCHAR* LexToString(ELobbyComparisonOp Comparison)
 {
 	switch (Comparison)
diff --git a/Engine/Plugins/Online/Experimental/OnlineServices/Source/OnlineServicesInterface/Private/Online/Schema.cpp b/Engine/Plugins/Online/Experimental/OnlineServices/Source/OnlineServicesInterface/Private/Online/Schema.cpp
--- a/Engine/Plugins/Online/Experimental/OnlineServices/Source/OnlineServicesInterface/Private/Online/Schema.cpp
+++ b/Engine/Plugins/Online/Experimental/OnlineServices/Source/OnlineServicesInterface/Private/Online/Schema.cpp
@@ -5,9 +5,9 @@
 namespace UE::Online {
 
 
-const FString LexToString(const FSchemaVariant& Variant)
+const FString LexToString(const FLexToStringAdaptor<FSchemaVariant>& Adaptor)
 {
-	return Variant.GetString();
+	return Adaptor.SchemaVariant.GetString();
 }
 
 void LexFromString(FSchemaVariant& Variant, const TCHAR* InStr)
diff --git a/Engine/Plugins/Online/Experimental/OnlineServices/Source/OnlineServicesInterface/Public/Online/Auth.h b/Engine/Plugins/Online/Experimental/OnlineServices/Source/OnlineServicesInterface/Public/Online/Auth.h
--- a/Engine/Plugins/Online/Experimental/OnlineServices/Source/OnlineServicesInterface/Public/Online/Auth.h
+++ b/Engine/Plugins/Online/Experimental/OnlineServices/Source/OnlineServicesInterface/Public/Online/Auth.h
@@ -520,7 +520,6 @@ public:
 };
 
 namespace Meta {
-// TODO: Move to Auth_Meta.inl file?
 
 BEGIN_ONLINE_STRUCT_META(FAccountInfo)
 	ONLINE_STRUCT_FIELD(FAccountInfo, AccountId),
diff --git a/Engine/Plugins/Online/Experimental/OnlineServices/Source/OnlineServicesInterface/Public/Online/Lobbies.h b/Engine/Plugins/Online/Experimental/OnlineServices/Source/OnlineServicesInterface/Public/Online/Lobbies.h
--- a/Engine/Plugins/Online/Experimental/OnlineServices/Source/OnlineServicesInterface/Public/Online/Lobbies.h
+++ b/Engine/Plugins/Online/Experimental/OnlineServices/Source/OnlineServicesInterface/Public/Online/Lobbies.h
@@ -48,12 +48,16 @@ enum class ELobbyMemberLeaveReason
 	*/
 	Closed
 };
+ONLINESERVICESINTERFACE_API const TCHAR* LexToString(ELobbyMemberLeaveReason LeaveReason);
+ONLINESERVICESINTERFACE_API void LexFromString(ELobbyMemberLeaveReason& OutLeaveReason, const TCHAR* InStr);
 
 enum class ELobbyAttributeVisibility
 {
 	Public,
 	Private,
 };
+ONLINESERVICESINTERFACE_API const TCHAR* LexToString(ELobbyAttributeVisibility Visibility);
+ONLINESERVICESINTERFACE_API void LexFromString(ELobbyAttributeVisibility& OutVisibility, const TCHAR* InStr);
 
 struct FLobbyMember
 {
@@ -986,7 +990,6 @@ BEGIN_ONLINE_STRUCT_META(FUILobbyJoinRequested)
 	ONLINE_STRUCT_FIELD(FUILobbyJoinRequested, Result)
 END_ONLINE_STRUCT_META()
 
-
 /* Meta*/ }
 
 /* UE::Online */ }
diff --git a/Engine/Plugins/Online/Experimental/OnlineServices/Source/OnlineServicesInterface/Public/Online/Schema.h b/Engine/Plugins/Online/Experimental/OnlineServices/Source/OnlineServicesInterface/Public/Online/Schema.h
--- a/Engine/Plugins/Online/Experimental/OnlineServices/Source/OnlineServicesInterface/Public/Online/Schema.h
+++ b/Engine/Plugins/Online/Experimental/OnlineServices/Source/OnlineServicesInterface/Public/Online/Schema.h
@@ -40,7 +40,22 @@ public:
 public:
 	FVariantType VariantData;
 };
-ONLINESERVICESINTERFACE_API const FString LexToString(const FSchemaVariant& Variant);
+
+// Don't allow implicit conversion to FSchemaVariant when calling LexToString.
+template <typename T> class FLexToStringAdaptor;
+template <>
+class FLexToStringAdaptor<FSchemaVariant>
+{
+public:
+	FLexToStringAdaptor(const FSchemaVariant& SchemaVariant)
+		: SchemaVariant(SchemaVariant)
+	{
+	}
+
+	const FSchemaVariant& SchemaVariant;
+};
+
+ONLINESERVICESINTERFACE_API const FString LexToString(const FLexToStringAdaptor<FSchemaVariant>& Adaptor);
 ONLINESERVICESINTERFACE_API void LexFromString(FSchemaVariant& Variant, const TCHAR* InStr);
 
 /* UE::Online */ }
diff --git a/Engine/Plugins/Online/Experimental/OnlineServicesEOS/Source/Private/Online/AuthEOS.cpp b/Engine/Plugins/Online/Experimental/OnlineServicesEOS/Source/Private/Online/AuthEOS.cpp
--- a/Engine/Plugins/Online/Experimental/OnlineServicesEOS/Source/Private/Online/AuthEOS.cpp
+++ b/Engine/Plugins/Online/Experimental/OnlineServicesEOS/Source/Private/Online/AuthEOS.cpp
@@ -88,7 +88,7 @@ TOnlineAsyncOpHandle<FAuthLogin> FAuthEOS::Login(FAuthLogin::Params&& Params)
 		FAuthEOSLoginConfig AuthEOSLoginConfig;
 		LoadConfig(AuthEOSLoginConfig, TEXT("Login"));
 
-		FLoginEASImpl::Params LoginParams;
+		FAuthLoginEASImpl::Params LoginParams;
 		LoginParams.PlatformUserId = Params.PlatformUserId;
 		LoginParams.CredentialsType = Params.CredentialsType;
 		LoginParams.CredentialsId = Params.CredentialsId;
@@ -96,7 +96,7 @@ TOnlineAsyncOpHandle<FAuthLogin> FAuthEOS::Login(FAuthLogin::Params&& Params)
 		LoginParams.Scopes = !Params.Scopes.IsEmpty() ? Params.Scopes : AuthEOSLoginConfig.DefaultScopes;
 
 		LoginEASImpl(LoginParams)
-		.Next([this, Promise = MoveTemp(Promise), WeakOp = InAsyncOp.AsWeak()](TDefaultErrorResult<FLoginEASImpl>&& LoginResult) mutable -> void
+		.Next([this, Promise = MoveTemp(Promise), WeakOp = InAsyncOp.AsWeak()](TDefaultErrorResult<FAuthLoginEASImpl>&& LoginResult) mutable -> void
 		{
 			if (TSharedPtr<TOnlineAsyncOp<FAuthLogin>> Op = WeakOp.Pin())
 			{
@@ -125,34 +125,34 @@ TOnlineAsyncOpHandle<FAuthLogin> FAuthEOS::Login(FAuthLogin::Params&& Params)
 		const FAuthLogin::Params& Params = InAsyncOp.GetParams();
 		const TSharedRef<FAccountInfoEOS>& AccountInfoEOS = GetOpDataChecked<TSharedRef<FAccountInfoEOS>>(InAsyncOp, AccountInfoKeyName);
 
-		TPromise<FLoginConnectImpl::Params> Promise;
-		TFuture<FLoginConnectImpl::Params> Future = Promise.GetFuture();
+		TPromise<FAuthLoginConnectImpl::Params> Promise;
+		TFuture<FAuthLoginConnectImpl::Params> Future = Promise.GetFuture();
 
-		TDefaultErrorResult<FGetExternalAuthTokenImpl> AuthTokenResult = GetExternalAuthTokenImpl(FGetExternalAuthTokenImpl::Params{AccountInfoEOS->EpicAccountId});
+		TDefaultErrorResult<FAuthGetExternalAuthTokenImpl> AuthTokenResult = GetExternalAuthTokenImpl(FAuthGetExternalAuthTokenImpl::Params{AccountInfoEOS->EpicAccountId});
 		if (AuthTokenResult.IsError())
 		{
 			UE_LOG(LogOnlineServices, Warning, TEXT("[FAuthEOS::Login] Failure: GetExternalAuthTokenImpl %s"), *AuthTokenResult.GetErrorValue().GetLogString());
 			InAsyncOp.SetError(Errors::Unknown(MoveTemp(AuthTokenResult.GetErrorValue())));
 
 			// Failed to acquire token - logout EAS.
-			LogoutEASImpl(FLogoutEASImpl::Params{ AccountInfoEOS->EpicAccountId })
-			.Next([Promise = MoveTemp(Promise)](TDefaultErrorResult<FLogoutEASImpl>&& LogoutResult) mutable -> void
+			LogoutEASImpl(FAuthLogoutEASImpl::Params{ AccountInfoEOS->EpicAccountId })
+			.Next([Promise = MoveTemp(Promise)](TDefaultErrorResult<FAuthLogoutEASImpl>&& LogoutResult) mutable -> void
 			{
 				if (LogoutResult.IsError())
 				{
 					UE_LOG(LogOnlineServices, Warning, TEXT("[FAuthEOS::Login] Failure: LogoutEASImpl %s"), *LogoutResult.GetErrorValue().GetLogString());
 				}
-				Promise.EmplaceValue(FLoginConnectImpl::Params{});
+				Promise.EmplaceValue(FAuthLoginConnectImpl::Params{});
 			});
 
 			return Future;
 		}
 
-		Promise.EmplaceValue(FLoginConnectImpl::Params{Params.PlatformUserId, MoveTemp(AuthTokenResult.GetOkValue().Token)});
+		Promise.EmplaceValue(FAuthLoginConnectImpl::Params{Params.PlatformUserId, MoveTemp(AuthTokenResult.GetOkValue().Token)});
 		return Future;
 	})
 	// Step 4: Attempt connect login. On connect login failure handle logout of EAS.
-	.Then([this](TOnlineAsyncOp<FAuthLogin>& InAsyncOp, FLoginConnectImpl::Params&& LoginConnectParams)
+	.Then([this](TOnlineAsyncOp<FAuthLogin>& InAsyncOp, FAuthLoginConnectImpl::Params&& LoginConnectParams)
 	{
 		const TSharedRef<FAccountInfoEOS>& AccountInfoEOS = GetOpDataChecked<TSharedRef<FAccountInfoEOS>>(InAsyncOp, AccountInfoKeyName);
 
@@ -161,7 +161,7 @@ TOnlineAsyncOpHandle<FAuthLogin> FAuthEOS::Login(FAuthLogin::Params&& Params)
 
 		// Attempt connect login.
 		LoginConnectImpl(LoginConnectParams)
-		.Next([this, AccountInfoEOS, WeakOp = InAsyncOp.AsWeak(), Promise = MoveTemp(Promise)](TDefaultErrorResult<FLoginConnectImpl>&& LoginResult) mutable -> void
+		.Next([this, AccountInfoEOS, WeakOp = InAsyncOp.AsWeak(), Promise = MoveTemp(Promise)](TDefaultErrorResult<FAuthLoginConnectImpl>&& LoginResult) mutable -> void
 		{
 			if (TSharedPtr<TOnlineAsyncOp<FAuthLogin>> Op = WeakOp.Pin())
 			{
@@ -170,8 +170,8 @@ TOnlineAsyncOpHandle<FAuthLogin> FAuthEOS::Login(FAuthLogin::Params&& Params)
 					UE_LOG(LogOnlineServices, Warning, TEXT("[FAuthEOS::Login] Failure: LoginConnectImpl %s"), *LoginResult.GetErrorValue().GetLogString());
 					Op->SetError(Errors::Unknown(MoveTemp(LoginResult.GetErrorValue())));
 
-					LogoutEASImpl(FLogoutEASImpl::Params{ AccountInfoEOS->EpicAccountId })
-					.Next([Promise = MoveTemp(Promise)](TDefaultErrorResult<FLogoutEASImpl>&& LogoutResult) mutable -> void
+					LogoutEASImpl(FAuthLogoutEASImpl::Params{ AccountInfoEOS->EpicAccountId })
+					.Next([Promise = MoveTemp(Promise)](TDefaultErrorResult<FAuthLogoutEASImpl>&& LogoutResult) mutable -> void
 					{
 						if (LogoutResult.IsError())
 						{
@@ -223,8 +223,8 @@ TOnlineAsyncOpHandle<FAuthLogin> FAuthEOS::Login(FAuthLogin::Params&& Params)
 			TPromise<void> Promise;
 			TFuture<void> Future = Promise.GetFuture();
 
-			LogoutEASImpl(FLogoutEASImpl::Params{ AccountInfoEOS->EpicAccountId })
-			.Next([Promise = MoveTemp(Promise)](TDefaultErrorResult<FLogoutEASImpl>&& LogoutResult) mutable -> void
+			LogoutEASImpl(FAuthLogoutEASImpl::Params{ AccountInfoEOS->EpicAccountId })
+			.Next([Promise = MoveTemp(Promise)](TDefaultErrorResult<FAuthLogoutEASImpl>&& LogoutResult) mutable -> void
 			{
 				if (LogoutResult.IsError())
 				{
@@ -330,7 +330,7 @@ TOnlineAsyncOpHandle<FAuthQueryExternalAuthToken> FAuthEOS::QueryExternalAuthTok
 			return;
 		}
 
-		TDefaultErrorResult<FGetExternalAuthTokenImpl> AuthTokenResult = GetExternalAuthTokenImpl(FGetExternalAuthTokenImpl::Params{ AccountInfoEOS->EpicAccountId });
+		TDefaultErrorResult<FAuthGetExternalAuthTokenImpl> AuthTokenResult = GetExternalAuthTokenImpl(FAuthGetExternalAuthTokenImpl::Params{ AccountInfoEOS->EpicAccountId });
 		if (AuthTokenResult.IsError())
 		{
 			UE_LOG(LogOnlineServices, Warning, TEXT("[FAuthEOS::QueryExternalAuthToken] Failure: GetExternalAuthTokenImpl %s"), *AuthTokenResult.GetErrorValue().GetLogString());
@@ -593,5 +593,4 @@ FOnlineAccountIdHandle FAuthEOS::CreateAccountId(const EOS_EpicAccountId EpicAcc
 	return FOnlineAccountIdRegistryEOS::Get().FindOrAddAccountId(EpicAccountId, ProductUserId);
 }
 
-
 /* UE::Online */ }
diff --git a/Engine/Plugins/Online/Experimental/OnlineServicesEOSGS/Source/Private/Online/AuthEOSGS.cpp b/Engine/Plugins/Online/Experimental/OnlineServicesEOSGS/Source/Private/Online/AuthEOSGS.cpp
--- a/Engine/Plugins/Online/Experimental/OnlineServicesEOSGS/Source/Private/Online/AuthEOSGS.cpp
+++ b/Engine/Plugins/Online/Experimental/OnlineServicesEOSGS/Source/Private/Online/AuthEOSGS.cpp
@@ -10,7 +10,6 @@
 #include "Online/OnlineErrorDefinitions.h"
 #include "Online/OnlineIdEOSGS.h"
 #include "Online/OnlineServicesEOSGS.h"
-#include "Online/OnlineServicesEOSGSTypes.h"
 
 #include "eos_auth.h"
 #include "eos_common.h"
@@ -615,7 +614,7 @@ TOnlineAsyncOpHandle<FAuthLogin> FAuthEOSGS::Login(FAuthLogin::Params&& Params)
 
 				const FAuthLogin::Params& Params = InAsyncOp.GetParams();
 
-				FLoginEASImpl::Params LoginParams;
+				FAuthLoginEASImpl::Params LoginParams;
 				LoginParams.PlatformUserId = Params.PlatformUserId;
 				LoginParams.CredentialsType = Params.CredentialsType;
 				LoginParams.CredentialsId = Params.CredentialsId;
@@ -623,7 +622,7 @@ TOnlineAsyncOpHandle<FAuthLogin> FAuthEOSGS::Login(FAuthLogin::Params&& Params)
 				LoginParams.Scopes = !Params.Scopes.IsEmpty() ? Params.Scopes : AuthEOSGSLoginConfig.DefaultScopes;
 
 				return LoginEASImpl(LoginParams)
-				.Next([this, Promise = MoveTemp(Promise), Op = InAsyncOp.AsShared()](TDefaultErrorResult<FLoginEASImpl>&& LoginResult) mutable -> void
+				.Next([this, Promise = MoveTemp(Promise), Op = InAsyncOp.AsShared()](TDefaultErrorResult<FAuthLoginEASImpl>&& LoginResult) mutable -> void
 				{
 					const TSharedRef<FAccountInfoEOS>& AccountInfoEOS = GetOpDataChecked<TSharedRef<FAccountInfoEOS>>(*Op, AccountInfoKeyName);
 
@@ -658,33 +657,33 @@ TOnlineAsyncOpHandle<FAuthLogin> FAuthEOSGS::Login(FAuthLogin::Params&& Params)
 		const FAuthLogin::Params& Params = InAsyncOp.GetParams();
 		const TSharedRef<FAccountInfoEOS>& AccountInfoEOS = GetOpDataChecked<TSharedRef<FAccountInfoEOS>>(InAsyncOp, AccountInfoKeyName);
 
-		TPromise<FLoginConnectImpl::Params> Promise;
-		TFuture<FLoginConnectImpl::Params> Future = Promise.GetFuture();
+		TPromise<FAuthLoginConnectImpl::Params> Promise;
+		TFuture<FAuthLoginConnectImpl::Params> Future = Promise.GetFuture();
 
 		// When an epic account is present, use the external auth credentials of the logged in account.
 		if (AccountInfoEOS->EpicAccountId)
 		{
-			TDefaultErrorResult<FGetExternalAuthTokenImpl> AuthTokenResult = GetExternalAuthTokenImpl(FGetExternalAuthTokenImpl::Params{AccountInfoEOS->EpicAccountId});
+			TDefaultErrorResult<FAuthGetExternalAuthTokenImpl> AuthTokenResult = GetExternalAuthTokenImpl(FAuthGetExternalAuthTokenImpl::Params{AccountInfoEOS->EpicAccountId});
 			if (AuthTokenResult.IsError())
 			{
 				UE_LOG(LogOnlineServices, Warning, TEXT("[FAuthEOSGS::Login] Failure: GetExternalAuthTokenImpl %s"), *AuthTokenResult.GetErrorValue().GetLogString());
 				InAsyncOp.SetError(Errors::Unknown(MoveTemp(AuthTokenResult.GetErrorValue())));
 
 				// Failed to acquire token - logout EAS.
-				LogoutEASImpl(FLogoutEASImpl::Params{ AccountInfoEOS->EpicAccountId })
-				.Next([Promise = MoveTemp(Promise)](TDefaultErrorResult<FLogoutEASImpl>&& LogoutResult) mutable -> void
+				LogoutEASImpl(FAuthLogoutEASImpl::Params{ AccountInfoEOS->EpicAccountId })
+				.Next([Promise = MoveTemp(Promise)](TDefaultErrorResult<FAuthLogoutEASImpl>&& LogoutResult) mutable -> void
 				{
 					if (LogoutResult.IsError())
 					{
 						UE_LOG(LogOnlineServices, Warning, TEXT("[FAuthEOSGS::Login] Failure: LogoutEASImpl %s"), *LogoutResult.GetErrorValue().GetLogString());
 					}
-					Promise.EmplaceValue(FLoginConnectImpl::Params{});
+					Promise.EmplaceValue(FAuthLoginConnectImpl::Params{});
 				});
 
 				return Future;
 			}
 
-			Promise.EmplaceValue(FLoginConnectImpl::Params{Params.PlatformUserId, MoveTemp(AuthTokenResult.GetOkValue().Token)});
+			Promise.EmplaceValue(FAuthLoginConnectImpl::Params{Params.PlatformUserId, MoveTemp(AuthTokenResult.GetOkValue().Token)});
 		}
 		else
 		{
@@ -692,17 +691,17 @@ TOnlineAsyncOpHandle<FAuthLogin> FAuthEOSGS::Login(FAuthLogin::Params&& Params)
 			if (Params.CredentialsType != LoginCredentialsType::ExternalAuth || !Params.CredentialsToken.IsType<FExternalAuthToken>())
 			{
 				InAsyncOp.SetError(Errors::InvalidParams());
-				Promise.EmplaceValue(FLoginConnectImpl::Params{});
+				Promise.EmplaceValue(FAuthLoginConnectImpl::Params{});
 				return Future;
 			}
 
-			Promise.EmplaceValue(FLoginConnectImpl::Params{Params.PlatformUserId, Params.CredentialsToken.Get<FExternalAuthToken>()});
+			Promise.EmplaceValue(FAuthLoginConnectImpl::Params{Params.PlatformUserId, Params.CredentialsToken.Get<FExternalAuthToken>()});
 		}
 
 		return Future;
 	})
 	// Step 4: Attempt connect login. On connect login failure handle logout of EAS.
-	.Then([this](TOnlineAsyncOp<FAuthLogin>& InAsyncOp, FLoginConnectImpl::Params&& LoginConnectParams)
+	.Then([this](TOnlineAsyncOp<FAuthLogin>& InAsyncOp, FAuthLoginConnectImpl::Params&& LoginConnectParams)
 	{
 		const TSharedRef<FAccountInfoEOS>& AccountInfoEOS = GetOpDataChecked<TSharedRef<FAccountInfoEOS>>(InAsyncOp, AccountInfoKeyName);
 
@@ -711,7 +710,7 @@ TOnlineAsyncOpHandle<FAuthLogin> FAuthEOSGS::Login(FAuthLogin::Params&& Params)
 
 		// Attempt connect login.
 		LoginConnectImpl(LoginConnectParams)
-		.Next([this, AccountInfoEOS, Op = InAsyncOp.AsShared(), Promise = MoveTemp(Promise)](TDefaultErrorResult<FLoginConnectImpl>&& LoginResult) mutable -> void
+		.Next([this, AccountInfoEOS, Op = InAsyncOp.AsShared(), Promise = MoveTemp(Promise)](TDefaultErrorResult<FAuthLoginConnectImpl>&& LoginResult) mutable -> void
 		{
 			if (LoginResult.IsError())
 			{
@@ -721,8 +720,8 @@ TOnlineAsyncOpHandle<FAuthLogin> FAuthEOSGS::Login(FAuthLogin::Params&& Params)
 				// Logout of EAS on login failure if necessary.
 				if (AccountInfoEOS->EpicAccountId)
 				{
-					LogoutEASImpl(FLogoutEASImpl::Params{ AccountInfoEOS->EpicAccountId })
-					.Next([Promise = MoveTemp(Promise)](TDefaultErrorResult<FLogoutEASImpl>&& LogoutResult) mutable -> void
+					LogoutEASImpl(FAuthLogoutEASImpl::Params{ AccountInfoEOS->EpicAccountId })
+					.Next([Promise = MoveTemp(Promise)](TDefaultErrorResult<FAuthLogoutEASImpl>&& LogoutResult) mutable -> void
 					{
 						if (LogoutResult.IsError())
 						{
@@ -773,8 +772,8 @@ TOnlineAsyncOpHandle<FAuthLogin> FAuthEOSGS::Login(FAuthLogin::Params&& Params)
 					TPromise<void> Promise;
 					TFuture<void> Future = Promise.GetFuture();
 
-					LogoutEASImpl(FLogoutEASImpl::Params{ AccountInfoEOS->EpicAccountId })
-					.Next([Promise = MoveTemp(Promise)](TDefaultErrorResult<FLogoutEASImpl>&& LogoutResult) mutable -> void
+					LogoutEASImpl(FAuthLogoutEASImpl::Params{ AccountInfoEOS->EpicAccountId })
+					.Next([Promise = MoveTemp(Promise)](TDefaultErrorResult<FAuthLogoutEASImpl>&& LogoutResult) mutable -> void
 					{
 						if (LogoutResult.IsError())
 						{
@@ -876,8 +875,8 @@ TOnlineAsyncOpHandle<FAuthLogout> FAuthEOSGS::Logout(FAuthLogout::Params&& Param
 			TPromise<void> Promise;
 			TFuture<void> Future = Promise.GetFuture();
 
-			LogoutEASImpl(FLogoutEASImpl::Params{ AccountInfoEOS->EpicAccountId })
-			.Next([this, AccountInfoEOS, Promise = MoveTemp(Promise)](TDefaultErrorResult<FLogoutEASImpl>&& LogoutResult) mutable -> void
+			LogoutEASImpl(FAuthLogoutEASImpl::Params{ AccountInfoEOS->EpicAccountId })
+			.Next([this, AccountInfoEOS, Promise = MoveTemp(Promise)](TDefaultErrorResult<FAuthLogoutEASImpl>&& LogoutResult) mutable -> void
 			{
 				if (LogoutResult.IsError())
 				{
@@ -983,18 +982,18 @@ TFunction<TFuture<TArray<FOnlineAccountIdHandle>>(FOnlineAsyncOp& InAsyncOp, con
 	};
 }
 
-TFuture<TDefaultErrorResult<FAuthEOSGS::FLoginEASImpl>> FAuthEOSGS::LoginEASImpl(const FLoginEASImpl::Params& LoginParams)
+TFuture<TDefaultErrorResult<FAuthLoginEASImpl>> FAuthEOSGS::LoginEASImpl(const FAuthLoginEASImpl::Params& LoginParams)
 {
 	TDefaultErrorResultInternal<FEOSAuthLoginOptions> LoginOptionsResult = FEOSAuthLoginOptions::Create(LoginParams.CredentialsType, LoginParams.CredentialsId, LoginParams.CredentialsToken, LoginParams.Scopes);
 	if (LoginOptionsResult.IsError())
 	{
-		return MakeFulfilledPromise<TDefaultErrorResult<FLoginEASImpl>>(MoveTemp(LoginOptionsResult.GetErrorValue())).GetFuture();
+		return MakeFulfilledPromise<TDefaultErrorResult<FAuthLoginEASImpl>>(MoveTemp(LoginOptionsResult.GetErrorValue())).GetFuture();
 	}
 
 	const bool IsPersistentAuthLogin = LoginParams.CredentialsType == LoginCredentialsType::PersistentAuth;
 
-	TPromise<TDefaultErrorResult<FLoginEASImpl>> Promise;
-	TFuture<TDefaultErrorResult<FLoginEASImpl>> Future = Promise.GetFuture();
+	TPromise<TDefaultErrorResult<FAuthLoginEASImpl>> Promise;
+	TFuture<TDefaultErrorResult<FAuthLoginEASImpl>> Future = Promise.GetFuture();
 
 	EOS_Async(EOS_Auth_Login, AuthHandle, MoveTemp(LoginOptionsResult.GetOkValue()),
 	[AuthHandle = AuthHandle, IsPersistentAuthLogin, Promise = MoveTemp(Promise)](const EOS_Auth_LoginCallbackInfo* Data) mutable -> void
@@ -1003,7 +1002,7 @@ TFuture<TDefaultErrorResult<FAuthEOSGS::FLoginEASImpl>> FAuthEOSGS::LoginEASImpl
 
 		if (Data->ResultCode == EOS_EResult::EOS_Success)
 		{
-			Promise.SetValue(TDefaultErrorResult<FLoginEASImpl>(FLoginEASImpl::Result{ Data->LocalUserId }));
+			Promise.SetValue(TDefaultErrorResult<FAuthLoginEASImpl>(FAuthLoginEASImpl::Result{ Data->LocalUserId }));
 		}
 		else if (Data->ResultCode == EOS_EResult::EOS_InvalidUser && Data->ContinuanceToken != nullptr)
 		{
@@ -1018,11 +1017,11 @@ TFuture<TDefaultErrorResult<FAuthEOSGS::FLoginEASImpl>> FAuthEOSGS::LoginEASImpl
 
 				if (Data->ResultCode == EOS_EResult::EOS_Success)
 				{
-					Promise.SetValue(TDefaultErrorResult<FLoginEASImpl>(FLoginEASImpl::Result{ Data->LocalUserId }));
+					Promise.SetValue(TDefaultErrorResult<FAuthLoginEASImpl>(FAuthLoginEASImpl::Result{ Data->LocalUserId }));
 				}
 				else
 				{
-					Promise.SetValue(TDefaultErrorResult<FLoginEASImpl>(FromEOSError(Data->ResultCode)));
+					Promise.SetValue(TDefaultErrorResult<FAuthLoginEASImpl>(FromEOSError(Data->ResultCode)));
 				}
 			});
 		}
@@ -1050,12 +1049,12 @@ TFuture<TDefaultErrorResult<FAuthEOSGS::FLoginEASImpl>> FAuthEOSGS::LoginEASImpl
 						UE_LOG(LogOnlineServices, Warning, TEXT("[FAuthEOS::LoginEASImpl] Failure: DeletePersistentAuthResult %s"), *DeletePersistentAuthError.GetLogString());
 					}
 
-					Promise.SetValue(TDefaultErrorResult<FLoginEASImpl>(MoveTemp(ResolvedError)));
+					Promise.SetValue(TDefaultErrorResult<FAuthLoginEASImpl>(MoveTemp(ResolvedError)));
 				});
 			}
 			else
 			{
-				Promise.SetValue(TDefaultErrorResult<FLoginEASImpl>(MoveTemp(ResolvedError)));
+				Promise.SetValue(TDefaultErrorResult<FAuthLoginEASImpl>(MoveTemp(ResolvedError)));
 			}
 		}
 	});
@@ -1063,10 +1062,10 @@ TFuture<TDefaultErrorResult<FAuthEOSGS::FLoginEASImpl>> FAuthEOSGS::LoginEASImpl
 	return Future;
 }
 
-TFuture<TDefaultErrorResult<FAuthEOSGS::FLogoutEASImpl>> FAuthEOSGS::LogoutEASImpl(const FLogoutEASImpl::Params& LogoutParams)
+TFuture<TDefaultErrorResult<FAuthLogoutEASImpl>> FAuthEOSGS::LogoutEASImpl(const FAuthLogoutEASImpl::Params& LogoutParams)
 {
-	TPromise<TDefaultErrorResult<FLogoutEASImpl>> Promise;
-	TFuture<TDefaultErrorResult<FLogoutEASImpl>> Future = Promise.GetFuture();
+	TPromise<TDefaultErrorResult<FAuthLogoutEASImpl>> Promise;
+	TFuture<TDefaultErrorResult<FAuthLogoutEASImpl>> Future = Promise.GetFuture();
 
 	EOS_Auth_LogoutOptions LogoutOptions = {};
 	LogoutOptions.ApiVersion = EOS_AUTH_LOGOUT_API_LATEST;
@@ -1079,18 +1078,18 @@ TFuture<TDefaultErrorResult<FAuthEOSGS::FLogoutEASImpl>> FAuthEOSGS::LogoutEASIm
 
 		if (Data->ResultCode == EOS_EResult::EOS_Success)
 		{
-			Promise.SetValue(TDefaultErrorResult<FLogoutEASImpl>(FLogoutEASImpl::Result{}));
+			Promise.SetValue(TDefaultErrorResult<FAuthLogoutEASImpl>(FAuthLogoutEASImpl::Result{}));
 		}
 		else
 		{
-			Promise.SetValue(TDefaultErrorResult<FLogoutEASImpl>(FromEOSError(Data->ResultCode)));
+			Promise.SetValue(TDefaultErrorResult<FAuthLogoutEASImpl>(FromEOSError(Data->ResultCode)));
 		}
 	});
 
 	return Future;
 }
 
-TDefaultErrorResult<FAuthEOSGS::FGetExternalAuthTokenImpl> FAuthEOSGS::GetExternalAuthTokenImpl(const FGetExternalAuthTokenImpl::Params& Params)
+TDefaultErrorResult<FAuthGetExternalAuthTokenImpl> FAuthEOSGS::GetExternalAuthTokenImpl(const FAuthGetExternalAuthTokenImpl::Params& Params)
 {
 	EOS_Auth_CopyIdTokenOptions CopyIdTokenOptions = {};
 	CopyIdTokenOptions.ApiVersion = EOS_AUTH_COPYIDTOKEN_API_LATEST;
@@ -1108,24 +1107,24 @@ TDefaultErrorResult<FAuthEOSGS::FGetExternalAuthTokenImpl> FAuthEOSGS::GetExtern
 		FExternalAuthToken ExternalAuthToken;
 		ExternalAuthToken.Type = ExternalLoginType::EpicIdToken;
 		ExternalAuthToken.Data = UTF8_TO_TCHAR(IdToken->JsonWebToken);
-		return TDefaultErrorResult<FGetExternalAuthTokenImpl>(FGetExternalAuthTokenImpl::Result{ MoveTemp(ExternalAuthToken) });
+		return TDefaultErrorResult<FAuthGetExternalAuthTokenImpl>(FAuthGetExternalAuthTokenImpl::Result{ MoveTemp(ExternalAuthToken) });
 	}
 	else
 	{
-		return TDefaultErrorResult<FGetExternalAuthTokenImpl>(FromEOSError(Result));
+		return TDefaultErrorResult<FAuthGetExternalAuthTokenImpl>(FromEOSError(Result));
 	}
 }
 
-TFuture<TDefaultErrorResult<FAuthEOSGS::FLoginConnectImpl>> FAuthEOSGS::LoginConnectImpl(const FLoginConnectImpl::Params& LoginParams)
+TFuture<TDefaultErrorResult<FAuthLoginConnectImpl>> FAuthEOSGS::LoginConnectImpl(const FAuthLoginConnectImpl::Params& LoginParams)
 {
 	TDefaultErrorResultInternal<FEOSConnectLoginOptions> LoginOptionsResult = FEOSConnectLoginOptions::Create(LoginParams.PlatformUserId, LoginParams.ExternalAuthToken);
 	if (LoginOptionsResult.IsError())
 	{
-		return MakeFulfilledPromise<TDefaultErrorResult<FLoginConnectImpl>>(MoveTemp(LoginOptionsResult.GetErrorValue())).GetFuture();
+		return MakeFulfilledPromise<TDefaultErrorResult<FAuthLoginConnectImpl>>(MoveTemp(LoginOptionsResult.GetErrorValue())).GetFuture();
 	}
 
-	TPromise<TDefaultErrorResult<FLoginConnectImpl>> Promise;
-	TFuture<TDefaultErrorResult<FLoginConnectImpl>> Future = Promise.GetFuture();
+	TPromise<TDefaultErrorResult<FAuthLoginConnectImpl>> Promise;
+	TFuture<TDefaultErrorResult<FAuthLoginConnectImpl>> Future = Promise.GetFuture();
 
 	EOS_Async(EOS_Connect_Login, ConnectHandle, MoveTemp(LoginOptionsResult.GetOkValue()),
 	[ConnectHandle = ConnectHandle, Promise = MoveTemp(Promise)](const EOS_Connect_LoginCallbackInfo* Data) mutable -> void
@@ -1134,7 +1133,7 @@ TFuture<TDefaultErrorResult<FAuthEOSGS::FLoginConnectImpl>> FAuthEOSGS::LoginCon
 
 		if (Data->ResultCode == EOS_EResult::EOS_Success)
 		{
-			Promise.SetValue(TDefaultErrorResult<FLoginConnectImpl>(FLoginConnectImpl::Result{ Data->LocalUserId }));
+			Promise.SetValue(TDefaultErrorResult<FAuthLoginConnectImpl>(FAuthLoginConnectImpl::Result{ Data->LocalUserId }));
 		}
 		else if (Data->ResultCode == EOS_EResult::EOS_InvalidUser && Data->ContinuanceToken != nullptr)
 		{
@@ -1149,31 +1148,31 @@ TFuture<TDefaultErrorResult<FAuthEOSGS::FLoginConnectImpl>> FAuthEOSGS::LoginCon
 
 				if (Data->ResultCode == EOS_EResult::EOS_Success)
 				{
-					Promise.SetValue(TDefaultErrorResult<FLoginConnectImpl>(FLoginConnectImpl::Result{ Data->LocalUserId }));
+					Promise.SetValue(TDefaultErrorResult<FAuthLoginConnectImpl>(FAuthLoginConnectImpl::Result{ Data->LocalUserId }));
 				}
 				else
 				{
-					Promise.SetValue(TDefaultErrorResult<FLoginConnectImpl>(FromEOSError(Data->ResultCode)));
+					Promise.SetValue(TDefaultErrorResult<FAuthLoginConnectImpl>(FromEOSError(Data->ResultCode)));
 				}
 			});
 		}
 		else
 		{
-			Promise.SetValue(TDefaultErrorResult<FLoginConnectImpl>(FromEOSError(Data->ResultCode)));
+			Promise.SetValue(TDefaultErrorResult<FAuthLoginConnectImpl>(FromEOSError(Data->ResultCode)));
 		}
 	});
 
 	return Future;
 }
 
-TOnlineAsyncOpHandle<FAuthEOSGS::FConnectLoginRecoveryImpl> FAuthEOSGS::ConnectLoginRecoveryImplOp(FConnectLoginRecoveryImpl::Params&& Params)
+TOnlineAsyncOpHandle<FAuthConnectLoginRecoveryImpl> FAuthEOSGS::ConnectLoginRecoveryImplOp(FAuthConnectLoginRecoveryImpl::Params&& Params)
 {
-	TOnlineAsyncOpRef<FConnectLoginRecoveryImpl> Op = GetOp<FConnectLoginRecoveryImpl>(MoveTemp(Params));
+	TOnlineAsyncOpRef<FAuthConnectLoginRecoveryImpl> Op = GetOp<FAuthConnectLoginRecoveryImpl>(MoveTemp(Params));
 
 	// Step 1: Setup operation data.
-	Op->Then([this](TOnlineAsyncOp<FConnectLoginRecoveryImpl>& InAsyncOp)
+	Op->Then([this](TOnlineAsyncOp<FAuthConnectLoginRecoveryImpl>& InAsyncOp)
 	{
-		const FConnectLoginRecoveryImpl::Params& Params = InAsyncOp.GetParams();
+		const FAuthConnectLoginRecoveryImpl::Params& Params = InAsyncOp.GetParams();
 		TSharedPtr<FAccountInfoEOS> AccountInfoEOS = AccountInfoRegistryEOS.Find(Params.LocalUserId);
 		if (!AccountInfoEOS)
 		{
@@ -1197,11 +1196,11 @@ TOnlineAsyncOpHandle<FAuthEOSGS::FConnectLoginRecoveryImpl> FAuthEOSGS::ConnectL
 		InAsyncOp.Data.Set<TSharedRef<FAccountInfoEOS>>(AccountInfoKeyName, AccountInfoEOS.ToSharedRef());
 	})
 	// Step 2: Acquire EAS external auth token.
-	.Then([this](TOnlineAsyncOp<FConnectLoginRecoveryImpl>& InAsyncOp)
+	.Then([this](TOnlineAsyncOp<FAuthConnectLoginRecoveryImpl>& InAsyncOp)
 	{
 		const TSharedRef<FAccountInfoEOS>& AccountInfoEOS = GetOpDataChecked<TSharedRef<FAccountInfoEOS>>(InAsyncOp, AccountInfoKeyName);
 
-		TDefaultErrorResult<FGetExternalAuthTokenImpl> AuthTokenResult = GetExternalAuthTokenImpl(FGetExternalAuthTokenImpl::Params{ AccountInfoEOS->EpicAccountId });
+		TDefaultErrorResult<FAuthGetExternalAuthTokenImpl> AuthTokenResult = GetExternalAuthTokenImpl(FAuthGetExternalAuthTokenImpl::Params{ AccountInfoEOS->EpicAccountId });
 		if (AuthTokenResult.IsError())
 		{
 			UE_LOG(LogOnlineServices, Warning, TEXT("[FAuthEOSGS::ConnectLoginRecoveryImplOp] Failure: GetExternalAuthTokenImpl %s"), *AuthTokenResult.GetErrorValue().GetLogString());
@@ -1209,13 +1208,13 @@ TOnlineAsyncOpHandle<FAuthEOSGS::FConnectLoginRecoveryImpl> FAuthEOSGS::ConnectL
 
 			// Reinitialize recovery timer.
 			InitializeConnectLoginRecoveryTimer(AccountInfoEOS);
-			return FLoginConnectImpl::Params{};
+			return FAuthLoginConnectImpl::Params{};
 		}
 
-		return FLoginConnectImpl::Params{ AccountInfoEOS->PlatformUserId, MoveTemp(AuthTokenResult.GetOkValue().Token) };
+		return FAuthLoginConnectImpl::Params{ AccountInfoEOS->PlatformUserId, MoveTemp(AuthTokenResult.GetOkValue().Token) };
 	})
 	// Step 3: Refresh connect login.
-	.Then([this](TOnlineAsyncOp<FConnectLoginRecoveryImpl>& InAsyncOp, FLoginConnectImpl::Params&& LoginConnectParams)
+	.Then([this](TOnlineAsyncOp<FAuthConnectLoginRecoveryImpl>& InAsyncOp, FAuthLoginConnectImpl::Params&& LoginConnectParams)
 	{
 		const TSharedRef<FAccountInfoEOS>& AccountInfoEOS = GetOpDataChecked<TSharedRef<FAccountInfoEOS>>(InAsyncOp, AccountInfoKeyName);
 
@@ -1224,7 +1223,7 @@ TOnlineAsyncOpHandle<FAuthEOSGS::FConnectLoginRecoveryImpl> FAuthEOSGS::ConnectL
 
 		// Attempt connect login.
 		LoginConnectImpl(LoginConnectParams)
-		.Next([this, AccountInfoEOS, Op = InAsyncOp.AsShared(), Promise = MoveTemp(Promise)](TDefaultErrorResult<FLoginConnectImpl>&& LoginResult) mutable -> void
+		.Next([this, AccountInfoEOS, Op = InAsyncOp.AsShared(), Promise = MoveTemp(Promise)](TDefaultErrorResult<FAuthLoginConnectImpl>&& LoginResult) mutable -> void
 		{
 			if (LoginResult.IsError())
 			{
@@ -1245,7 +1244,7 @@ TOnlineAsyncOpHandle<FAuthEOSGS::FConnectLoginRecoveryImpl> FAuthEOSGS::ConnectL
 		return Future;
 	})
 	// Step 4: Update cache and notify.
-	.Then([this](TOnlineAsyncOp<FConnectLoginRecoveryImpl>& InAsyncOp)
+	.Then([this](TOnlineAsyncOp<FAuthConnectLoginRecoveryImpl>& InAsyncOp)
 	{
 		const TSharedRef<FAccountInfoEOS>& AccountInfoEOS = GetOpDataChecked<TSharedRef<FAccountInfoEOS>>(InAsyncOp, AccountInfoKeyName);
 
@@ -1253,21 +1252,21 @@ TOnlineAsyncOpHandle<FAuthEOSGS::FConnectLoginRecoveryImpl> FAuthEOSGS::ConnectL
 		AccountInfoEOS->LoginStatus = ELoginStatus::LoggedIn;
 		OnAuthLoginStatusChangedEvent.Broadcast(FAuthLoginStatusChanged{ AccountInfoEOS, AccountInfoEOS->LoginStatus });
 
-		InAsyncOp.SetResult(FConnectLoginRecoveryImpl::Result{});
+		InAsyncOp.SetResult(FAuthConnectLoginRecoveryImpl::Result{});
 	})
 	.Enqueue(GetSerialQueue());
 
 	return Op->GetHandle();
 }
 
-TOnlineAsyncOpHandle<FAuthEOSGS::FHandleConnectLoginStatusChangedImpl> FAuthEOSGS::HandleConnectLoginStatusChangedImplOp(FHandleConnectLoginStatusChangedImpl::Params&& Params)
+TOnlineAsyncOpHandle<FAuthHandleConnectLoginStatusChangedImpl> FAuthEOSGS::HandleConnectLoginStatusChangedImplOp(FAuthHandleConnectLoginStatusChangedImpl::Params&& Params)
 {
-	TOnlineAsyncOpRef<FHandleConnectLoginStatusChangedImpl> Op = GetOp<FHandleConnectLoginStatusChangedImpl>(MoveTemp(Params));
+	TOnlineAsyncOpRef<FAuthHandleConnectLoginStatusChangedImpl> Op = GetOp<FAuthHandleConnectLoginStatusChangedImpl>(MoveTemp(Params));
 
 	// Step 1: Set up operation data.
-	Op->Then([this](TOnlineAsyncOp<FHandleConnectLoginStatusChangedImpl>& InAsyncOp)
+	Op->Then([this](TOnlineAsyncOp<FAuthHandleConnectLoginStatusChangedImpl>& InAsyncOp)
 	{
-		const FHandleConnectLoginStatusChangedImpl::Params& Params = InAsyncOp.GetParams();
+		const FAuthHandleConnectLoginStatusChangedImpl::Params& Params = InAsyncOp.GetParams();
 		TSharedPtr<FAccountInfoEOS> AccountInfoEOS = AccountInfoRegistryEOS.Find(Params.LocalUserId);
 		if (!AccountInfoEOS)
 		{
@@ -1279,9 +1278,9 @@ TOnlineAsyncOpHandle<FAuthEOSGS::FHandleConnectLoginStatusChangedImpl> FAuthEOSG
 		InAsyncOp.Data.Set<TSharedRef<FAccountInfoEOS>>(AccountInfoKeyName, AccountInfoEOS.ToSharedRef());
 	})
 	// Step 2: Update status and notify.
-	.Then([this](TOnlineAsyncOp<FHandleConnectLoginStatusChangedImpl>& InAsyncOp)
+	.Then([this](TOnlineAsyncOp<FAuthHandleConnectLoginStatusChangedImpl>& InAsyncOp)
 	{
-		const FHandleConnectLoginStatusChangedImpl::Params& Params = InAsyncOp.GetParams();
+		const FAuthHandleConnectLoginStatusChangedImpl::Params& Params = InAsyncOp.GetParams();
 		const TSharedRef<FAccountInfoEOS>& AccountInfoEOS = GetOpDataChecked<TSharedRef<FAccountInfoEOS>>(InAsyncOp, AccountInfoKeyName);
 
 		// Check if the user was kicked offline.
@@ -1306,21 +1305,21 @@ TOnlineAsyncOpHandle<FAuthEOSGS::FHandleConnectLoginStatusChangedImpl> FAuthEOSG
 			}
 		}
 
-		InAsyncOp.SetResult(FHandleConnectLoginStatusChangedImpl::Result{});
+		InAsyncOp.SetResult(FAuthHandleConnectLoginStatusChangedImpl::Result{});
 	})
 	.Enqueue(GetSerialQueue());
 
 	return Op->GetHandle();
 }
 
-TOnlineAsyncOpHandle<FAuthEOSGS::FHandleConnectAuthNotifyExpirationImpl> FAuthEOSGS::HandleConnectAuthNotifyExpirationImplOp(FHandleConnectAuthNotifyExpirationImpl::Params&& Params)
+TOnlineAsyncOpHandle<FAuthHandleConnectAuthNotifyExpirationImpl> FAuthEOSGS::HandleConnectAuthNotifyExpirationImplOp(FAuthHandleConnectAuthNotifyExpirationImpl::Params&& Params)
 {
-	TOnlineAsyncOpRef<FHandleConnectAuthNotifyExpirationImpl> Op = GetOp<FHandleConnectAuthNotifyExpirationImpl>(MoveTemp(Params));
+	TOnlineAsyncOpRef<FAuthHandleConnectAuthNotifyExpirationImpl> Op = GetOp<FAuthHandleConnectAuthNotifyExpirationImpl>(MoveTemp(Params));
 
 	// Step 1: Set up operation data.
-	Op->Then([this](TOnlineAsyncOp<FHandleConnectAuthNotifyExpirationImpl>& InAsyncOp)
+	Op->Then([this](TOnlineAsyncOp<FAuthHandleConnectAuthNotifyExpirationImpl>& InAsyncOp)
 	{
-		const FHandleConnectAuthNotifyExpirationImpl::Params& Params = InAsyncOp.GetParams();
+		const FAuthHandleConnectAuthNotifyExpirationImpl::Params& Params = InAsyncOp.GetParams();
 		TSharedPtr<FAccountInfoEOS> AccountInfoEOS = AccountInfoRegistryEOS.Find(Params.LocalUserId);
 		if (!AccountInfoEOS)
 		{
@@ -1338,9 +1337,9 @@ TOnlineAsyncOpHandle<FAuthEOSGS::FHandleConnectAuthNotifyExpirationImpl> FAuthEO
 		InAsyncOp.Data.Set<TSharedRef<FAccountInfoEOS>>(AccountInfoKeyName, AccountInfoEOS.ToSharedRef());
 	})
 	// Step 2: Automatically start recovery or notify user to start recovery.
-	.Then([this](TOnlineAsyncOp<FHandleConnectAuthNotifyExpirationImpl>& InAsyncOp)
+	.Then([this](TOnlineAsyncOp<FAuthHandleConnectAuthNotifyExpirationImpl>& InAsyncOp)
 	{
-		const FHandleConnectAuthNotifyExpirationImpl::Params& Params = InAsyncOp.GetParams();
+		const FAuthHandleConnectAuthNotifyExpirationImpl::Params& Params = InAsyncOp.GetParams();
 		const TSharedRef<FAccountInfoEOS>& AccountInfoEOS = GetOpDataChecked<TSharedRef<FAccountInfoEOS>>(InAsyncOp, AccountInfoKeyName);
 
 		// When EAS is logged in use it to refresh the connect login.
@@ -1355,21 +1354,21 @@ TOnlineAsyncOpHandle<FAuthEOSGS::FHandleConnectAuthNotifyExpirationImpl> FAuthEO
 			OnAuthPendingAuthExpirationEvent.Broadcast(FAuthPendingAuthExpiration{ AccountInfoEOS });
 		}
 
-		InAsyncOp.SetResult(FHandleConnectAuthNotifyExpirationImpl::Result{});
+		InAsyncOp.SetResult(FAuthHandleConnectAuthNotifyExpirationImpl::Result{});
 	})
 	.Enqueue(GetSerialQueue());
 
 	return Op->GetHandle();
 }
 
-TOnlineAsyncOpHandle<FAuthEOSGS::FHandleEASLoginStatusChangedImpl> FAuthEOSGS::HandleEASLoginStatusChangedImplOp(FHandleEASLoginStatusChangedImpl::Params&& Params)
+TOnlineAsyncOpHandle<FAuthHandleEASLoginStatusChangedImpl> FAuthEOSGS::HandleEASLoginStatusChangedImplOp(FAuthHandleEASLoginStatusChangedImpl::Params&& Params)
 {
-	TOnlineAsyncOpRef<FHandleEASLoginStatusChangedImpl> Op = GetOp<FHandleEASLoginStatusChangedImpl>(MoveTemp(Params));
+	TOnlineAsyncOpRef<FAuthHandleEASLoginStatusChangedImpl> Op = GetOp<FAuthHandleEASLoginStatusChangedImpl>(MoveTemp(Params));
 
 	// Step 1: Set up operation data.
-	Op->Then([this](TOnlineAsyncOp<FHandleEASLoginStatusChangedImpl>& InAsyncOp)
+	Op->Then([this](TOnlineAsyncOp<FAuthHandleEASLoginStatusChangedImpl>& InAsyncOp)
 	{
-		const FHandleEASLoginStatusChangedImpl::Params& Params = InAsyncOp.GetParams();
+		const FAuthHandleEASLoginStatusChangedImpl::Params& Params = InAsyncOp.GetParams();
 		TSharedPtr<FAccountInfoEOS> AccountInfoEOS = AccountInfoRegistryEOS.Find(Params.LocalUserId);
 		if (!AccountInfoEOS)
 		{
@@ -1387,9 +1386,9 @@ TOnlineAsyncOpHandle<FAuthEOSGS::FHandleEASLoginStatusChangedImpl> FAuthEOSGS::H
 		InAsyncOp.Data.Set<TSharedRef<FAccountInfoEOS>>(AccountInfoKeyName, AccountInfoEOS.ToSharedRef());
 	})
 	// Step 2: Update status and notify.
-	.Then([this](TOnlineAsyncOp<FHandleEASLoginStatusChangedImpl>& InAsyncOp)
+	.Then([this](TOnlineAsyncOp<FAuthHandleEASLoginStatusChangedImpl>& InAsyncOp)
 	{
-		const FHandleEASLoginStatusChangedImpl::Params& Params = InAsyncOp.GetParams();
+		const FAuthHandleEASLoginStatusChangedImpl::Params& Params = InAsyncOp.GetParams();
 		const TSharedRef<FAccountInfoEOS>& AccountInfoEOS = GetOpDataChecked<TSharedRef<FAccountInfoEOS>>(InAsyncOp, AccountInfoKeyName);
 
 		// Check if the user was kicked offline.
@@ -1401,7 +1400,7 @@ TOnlineAsyncOpHandle<FAuthEOSGS::FHandleEASLoginStatusChangedImpl> FAuthEOSGS::H
 			AccountInfoRegistryEOS.Unregister(AccountInfoEOS->AccountId);
 		}
 
-		InAsyncOp.SetResult(FHandleEASLoginStatusChangedImpl::Result{});
+		InAsyncOp.SetResult(FAuthHandleEASLoginStatusChangedImpl::Result{});
 	})
 	.Enqueue(GetSerialQueue());
 
@@ -1447,19 +1446,59 @@ void FAuthEOSGS::UnregisterHandlers()
 
 void FAuthEOSGS::OnConnectLoginStatusChanged(const EOS_Connect_LoginStatusChangedCallbackInfo* Data)
 {
-	HandleConnectLoginStatusChangedImplOp(FHandleConnectLoginStatusChangedImpl::Params{ Data->LocalUserId, Data->PreviousStatus, Data->CurrentStatus });
+	HandleConnectLoginStatusChangedImplOp(FAuthHandleConnectLoginStatusChangedImpl::Params{ Data->LocalUserId, Data->PreviousStatus, Data->CurrentStatus });
 }
 
 void FAuthEOSGS::OnConnectAuthNotifyExpiration(const EOS_Connect_AuthExpirationCallbackInfo* Data)
 {
-	HandleConnectAuthNotifyExpirationImplOp(FHandleConnectAuthNotifyExpirationImpl::Params{ Data->LocalUserId });
+	HandleConnectAuthNotifyExpirationImplOp(FAuthHandleConnectAuthNotifyExpirationImpl::Params{ Data->LocalUserId });
 }
 
 void FAuthEOSGS::OnEASLoginStatusChanged(const EOS_Auth_LoginStatusChangedCallbackInfo* Data)
 {
-	HandleEASLoginStatusChangedImplOp(FHandleEASLoginStatusChangedImpl::Params{ Data->LocalUserId, Data->PrevStatus, Data->CurrentStatus });
+	HandleEASLoginStatusChangedImplOp(FAuthHandleEASLoginStatusChangedImpl::Params{ Data->LocalUserId, Data->PrevStatus, Data->CurrentStatus });
 }
 
+#if !UE_BUILD_SHIPPING
+void FAuthEOSGS::CheckMetadata()
+{
+	// Metadata sanity check.
+	ToLogString(FAuthLoginEASImpl::Params());
+	ToLogString(FAuthLoginEASImpl::Result());
+	ToLogString(FAuthLogoutEASImpl::Params());
+	ToLogString(FAuthLogoutEASImpl::Result());
+	ToLogString(FAuthGetExternalAuthTokenImpl::Params());
+	ToLogString(FAuthGetExternalAuthTokenImpl::Result());
+	ToLogString(FAuthLoginConnectImpl::Params());
+	ToLogString(FAuthLoginConnectImpl::Result());
+	ToLogString(FAuthConnectLoginRecoveryImpl::Params());
+	ToLogString(FAuthConnectLoginRecoveryImpl::Result());
+	ToLogString(FAuthHandleConnectLoginStatusChangedImpl::Params());
+	ToLogString(FAuthHandleConnectLoginStatusChangedImpl::Result());
+	ToLogString(FAuthHandleConnectAuthNotifyExpirationImpl::Params());
+	ToLogString(FAuthHandleConnectAuthNotifyExpirationImpl::Result());
+	ToLogString(FAuthHandleEASLoginStatusChangedImpl::Params());
+	ToLogString(FAuthHandleEASLoginStatusChangedImpl::Result());
+	ToLogString(FAccountInfoEOS());
+	Meta::VisitFields(FAuthLoginEASImpl::Params(), [](const TCHAR* Name, auto& Field) { return false; });
+	Meta::VisitFields(FAuthLoginEASImpl::Result(), [](const TCHAR* Name, auto& Field) { return false; });
+	Meta::VisitFields(FAuthLogoutEASImpl::Params(), [](const TCHAR* Name, auto& Field) { return false; });
+	Meta::VisitFields(FAuthLogoutEASImpl::Result(), [](const TCHAR* Name, auto& Field) { return false; });
+	Meta::VisitFields(FAuthGetExternalAuthTokenImpl::Params(), [](const TCHAR* Name, auto& Field) { return false; });
+	Meta::VisitFields(FAuthGetExternalAuthTokenImpl::Result(), [](const TCHAR* Name, auto& Field) { return false; });
+	Meta::VisitFields(FAuthLoginConnectImpl::Params(), [](const TCHAR* Name, auto& Field) { return false; });
+	Meta::VisitFields(FAuthLoginConnectImpl::Result(), [](const TCHAR* Name, auto& Field) { return false; });
+	Meta::VisitFields(FAuthConnectLoginRecoveryImpl::Params(), [](const TCHAR* Name, auto& Field) { return false; });
+	Meta::VisitFields(FAuthConnectLoginRecoveryImpl::Result(), [](const TCHAR* Name, auto& Field) { return false; });
+	Meta::VisitFields(FAuthHandleConnectLoginStatusChangedImpl::Params(), [](const TCHAR* Name, auto& Field) { return false; });
+	Meta::VisitFields(FAuthHandleConnectLoginStatusChangedImpl::Result(), [](const TCHAR* Name, auto& Field) { return false; });
+	Meta::VisitFields(FAuthHandleConnectAuthNotifyExpirationImpl::Params(), [](const TCHAR* Name, auto& Field) { return false; });
+	Meta::VisitFields(FAuthHandleConnectAuthNotifyExpirationImpl::Result(), [](const TCHAR* Name, auto& Field) { return false; });
+	Meta::VisitFields(FAuthHandleEASLoginStatusChangedImpl::Params(), [](const TCHAR* Name, auto& Field) { return false; });
+	Meta::VisitFields(FAuthHandleEASLoginStatusChangedImpl::Result(), [](const TCHAR* Name, auto& Field) { return false; });
+}
+#endif
+
 const FAccountInfoRegistry& FAuthEOSGS::GetAccountInfoRegistry() const
 {
 	return AccountInfoRegistryEOS;
@@ -1479,7 +1518,7 @@ void FAuthEOSGS::InitializeConnectLoginRecoveryTimer(const TSharedRef<FAccountIn
 	AccountInfoEOS->RestoreLoginTimer = FTSBackgroundableTicker::GetCoreTicker().AddTicker(
 		FTickerDelegate::CreateLambda([this, EpicAccountId = AccountInfoEOS->EpicAccountId](float) -> bool
 		{
-			ConnectLoginRecoveryImplOp(FConnectLoginRecoveryImpl::Params{ EpicAccountId });
+			ConnectLoginRecoveryImplOp(FAuthConnectLoginRecoveryImpl::Params{ EpicAccountId });
 
 			// One-shot timer.
 			return false;
diff --git a/Engine/Plugins/Online/Experimental/OnlineServicesEOSGS/Source/Private/Online/LobbiesEOSGS.cpp b/Engine/Plugins/Online/Experimental/OnlineServicesEOSGS/Source/Private/Online/LobbiesEOSGS.cpp
--- a/Engine/Plugins/Online/Experimental/OnlineServicesEOSGS/Source/Private/Online/LobbiesEOSGS.cpp
+++ b/Engine/Plugins/Online/Experimental/OnlineServicesEOSGS/Source/Private/Online/LobbiesEOSGS.cpp
@@ -74,7 +74,7 @@ TOnlineAsyncOpHandle<FCreateLobby> FLobbiesEOSGS::CreateLobby(FCreateLobby::Para
 		FOnlineAccountIdHandle LocalUserId,
 		FOnlineError ErrorResult) -> TFuture<void>
 	{
-		FDestroyLobbyImpl::Params DestroyLobbyParams;
+		FLobbiesDestroyLobbyImpl::Params DestroyLobbyParams;
 		DestroyLobbyParams.LobbyData = LobbyData;
 		DestroyLobbyParams.LocalUserId = LocalUserId;
 
@@ -82,7 +82,7 @@ TOnlineAsyncOpHandle<FCreateLobby> FLobbiesEOSGS::CreateLobby(FCreateLobby::Para
 		auto Future = Promise.GetFuture();
 
 		DestroyLobbyImpl(MoveTemp(DestroyLobbyParams))
-		.Then([this, InAsyncOp = InAsyncOp.AsShared(), ErrorResult = MoveTemp(ErrorResult), Promise = MoveTemp(Promise)](TFuture<TDefaultErrorResult<FDestroyLobbyImpl>>&& Future) mutable
+		.Then([this, InAsyncOp = InAsyncOp.AsShared(), ErrorResult = MoveTemp(ErrorResult), Promise = MoveTemp(Promise)](TFuture<TDefaultErrorResult<FLobbiesDestroyLobbyImpl>>&& Future) mutable
 		{
 			if (Future.Get().IsError())
 			{
@@ -194,7 +194,7 @@ TOnlineAsyncOpHandle<FCreateLobby> FLobbiesEOSGS::CreateLobby(FCreateLobby::Para
 		{
 			LobbyOwnerAttributes->MutatedAttributes = UserData->Attributes;
 
-			FModifyLobbyMemberDataImpl::Params ModifyLobbyMemberDataParams;
+			FLobbiesModifyLobbyMemberDataImpl::Params ModifyLobbyMemberDataParams;
 			ModifyLobbyMemberDataParams.LobbyData = LobbyData;
 			ModifyLobbyMemberDataParams.LocalUserId = Params.LocalUserId;
 			ModifyLobbyMemberDataParams.Changes = LobbyOwnerAttributes;
@@ -203,10 +203,10 @@ TOnlineAsyncOpHandle<FCreateLobby> FLobbiesEOSGS::CreateLobby(FCreateLobby::Para
 		}
 		else
 		{
-			return MakeFulfilledPromise<TDefaultErrorResult<FModifyLobbyMemberDataImpl>>().GetFuture();
+			return MakeFulfilledPromise<TDefaultErrorResult<FLobbiesModifyLobbyMemberDataImpl>>().GetFuture();
 		}
 	})
-	.Then([this, DestroyLobbyDuringCreate](TOnlineAsyncOp<FCreateLobby>& InAsyncOp, TDefaultErrorResult<FModifyLobbyMemberDataImpl>&& ModifyLobbyOwnerResult)
+	.Then([this, DestroyLobbyDuringCreate](TOnlineAsyncOp<FCreateLobby>& InAsyncOp, TDefaultErrorResult<FLobbiesModifyLobbyMemberDataImpl>&& ModifyLobbyOwnerResult)
 	{
 		// Step 5: Handle result
 
@@ -243,14 +243,14 @@ TOnlineAsyncOpHandle<FCreateLobby> FLobbiesEOSGS::CreateLobby(FCreateLobby::Para
 
 		// Add lobby attributes.
 		// Set lobby privacy setting to the user provided value.
-		FModifyLobbyDataImpl::Params ModifyLobbyDataParams;
+		FLobbiesModifyLobbyDataImpl::Params ModifyLobbyDataParams;
 		ModifyLobbyDataParams.LobbyData = LobbyData;
 		ModifyLobbyDataParams.LocalUserId = Params.LocalUserId;
 		ModifyLobbyDataParams.Changes = LobbyChanges;
 
 		return ModifyLobbyDataImpl(MoveTemp(ModifyLobbyDataParams));
 	})
-	.Then([this, DestroyLobbyDuringCreate](TOnlineAsyncOp<FCreateLobby>& InAsyncOp, TDefaultErrorResult<FModifyLobbyDataImpl>&& ModifyLobbyResult) mutable
+	.Then([this, DestroyLobbyDuringCreate](TOnlineAsyncOp<FCreateLobby>& InAsyncOp, TDefaultErrorResult<FLobbiesModifyLobbyDataImpl>&& ModifyLobbyResult) mutable
 	{
 		// Step 8: Handle result.
 
@@ -358,14 +358,14 @@ TOnlineAsyncOpHandle<FJoinLobby> FLobbiesEOSGS::JoinLobby(FJoinLobby::Params&& I
 
 		const FJoinLobby::Params& Params = InAsyncOp.GetParams();
 
-		FJoinLobbyImpl::Params JoinParams;
+		FLobbiesJoinLobbyImpl::Params JoinParams;
 		JoinParams.LobbyData = LobbyData;
 		JoinParams.LocalUserId = Params.LocalUserId;
 		JoinParams.LocalName = Params.LocalName;
 		JoinParams.LocalUsers = Params.LocalUsers;
 		return JoinLobbyImpl(MoveTemp(JoinParams));
 	})
-	.Then([this](TOnlineAsyncOp<FJoinLobby>& InAsyncOp, TDefaultErrorResult<FJoinLobbyImpl> Result)
+	.Then([this](TOnlineAsyncOp<FJoinLobby>& InAsyncOp, TDefaultErrorResult<FLobbiesJoinLobbyImpl> Result)
 	{
 		// Handle result.
 		if (Result.IsError())
@@ -430,12 +430,12 @@ TOnlineAsyncOpHandle<FLeaveLobby> FLobbiesEOSGS::LeaveLobby(FLeaveLobby::Params&
 
 		const FLeaveLobby::Params& Params = InAsyncOp.GetParams();
 
-		FLeaveLobbyImpl::Params LeaveParams;
+		FLobbiesLeaveLobbyImpl::Params LeaveParams;
 		LeaveParams.LobbyData = LobbyData;
 		LeaveParams.LocalUserId = Params.LocalUserId;
 		return LeaveLobbyImpl(MoveTemp(LeaveParams));
 	})
-	.Then([this, LobbyData](TOnlineAsyncOp<FLeaveLobby>& InAsyncOp, TDefaultErrorResult<FLeaveLobbyImpl> Result)
+	.Then([this, LobbyData](TOnlineAsyncOp<FLeaveLobby>& InAsyncOp, TDefaultErrorResult<FLobbiesLeaveLobbyImpl> Result)
 	{
 		const FLeaveLobby::Params& Params = InAsyncOp.GetParams();
 
@@ -481,13 +481,13 @@ TOnlineAsyncOpHandle<FInviteLobbyMember> FLobbiesEOSGS::InviteLobbyMember(FInvit
 	{
 		const FInviteLobbyMember::Params& Params = InAsyncOp.GetParams();
 
-		FInviteLobbyMemberImpl::Params InviteParams;
+		FLobbiesInviteLobbyMemberImpl::Params InviteParams;
 		InviteParams.LobbyData = LobbyDataRegistry->Find(Params.LobbyId);
 		InviteParams.LocalUserId = Params.LocalUserId;
 		InviteParams.TargetUserId = Params.TargetUserId;
 		return InviteLobbyMemberImpl(MoveTemp(InviteParams));
 	})
-	.Then([this](TOnlineAsyncOp<FInviteLobbyMember>& InAsyncOp, TDefaultErrorResult<FInviteLobbyMemberImpl> Result)
+	.Then([this](TOnlineAsyncOp<FInviteLobbyMember>& InAsyncOp, TDefaultErrorResult<FLobbiesInviteLobbyMemberImpl> Result)
 	{
 		if (Result.IsError())
 		{
@@ -519,12 +519,12 @@ TOnlineAsyncOpHandle<FDeclineLobbyInvitation> FLobbiesEOSGS::DeclineLobbyInvitat
 	{
 		const FDeclineLobbyInvitation::Params& Params = InAsyncOp.GetParams();
 
-		FDeclineLobbyInvitationImpl::Params DeclineParams;
+		FLobbiesDeclineLobbyInvitationImpl::Params DeclineParams;
 		DeclineParams.LocalUserId = Params.LocalUserId;
 		DeclineParams.LobbyId = Params.LobbyId;
 		return DeclineLobbyInvitationImpl(MoveTemp(DeclineParams));
 	})
-	.Then([this](TOnlineAsyncOp<FDeclineLobbyInvitation>& InAsyncOp, TDefaultErrorResult<FDeclineLobbyInvitationImpl> Result)
+	.Then([this](TOnlineAsyncOp<FDeclineLobbyInvitation>& InAsyncOp, TDefaultErrorResult<FLobbiesDeclineLobbyInvitationImpl> Result)
 	{
 		if (Result.IsError())
 		{
@@ -565,13 +565,13 @@ TOnlineAsyncOpHandle<FKickLobbyMember> FLobbiesEOSGS::KickLobbyMember(FKickLobby
 		const FKickLobbyMember::Params& Params = InAsyncOp.GetParams();
 		const TSharedRef<FLobbyDataEOS>& LobbyData = GetOpDataChecked<TSharedRef<FLobbyDataEOS>>(InAsyncOp, LobbyDataKeyName);
 
-		FKickLobbyMemberImpl::Params KickParams;
+		FLobbiesKickLobbyMemberImpl::Params KickParams;
 		KickParams.LobbyData = LobbyData;
 		KickParams.LocalUserId = Params.LocalUserId;
 		KickParams.TargetUserId = Params.TargetUserId;
 		return KickLobbyMemberImpl(MoveTemp(KickParams));
 	})
-	.Then([this](TOnlineAsyncOp<FKickLobbyMember>& InAsyncOp, TDefaultErrorResult<FKickLobbyMemberImpl> Result)
+	.Then([this](TOnlineAsyncOp<FKickLobbyMember>& InAsyncOp, TDefaultErrorResult<FLobbiesKickLobbyMemberImpl> Result)
 	{
 		if (Result.IsError())
 		{
@@ -619,13 +619,13 @@ TOnlineAsyncOpHandle<FPromoteLobbyMember> FLobbiesEOSGS::PromoteLobbyMember(FPro
 		const FPromoteLobbyMember::Params& Params = InAsyncOp.GetParams();
 		const TSharedRef<FLobbyDataEOS>& LobbyData = GetOpDataChecked<TSharedRef<FLobbyDataEOS>>(InAsyncOp, LobbyDataKeyName);
 
-		FPromoteLobbyMemberImpl::Params PromoteParams;
+		FLobbiesPromoteLobbyMemberImpl::Params PromoteParams;
 		PromoteParams.LobbyData = LobbyData;
 		PromoteParams.LocalUserId = Params.LocalUserId;
 		PromoteParams.TargetUserId = Params.TargetUserId;
 		return PromoteLobbyMemberImpl(MoveTemp(PromoteParams));
 	})
-	.Then([this](TOnlineAsyncOp<FPromoteLobbyMember>& InAsyncOp, TDefaultErrorResult<FPromoteLobbyMemberImpl> Result)
+	.Then([this](TOnlineAsyncOp<FPromoteLobbyMember>& InAsyncOp, TDefaultErrorResult<FLobbiesPromoteLobbyMemberImpl> Result)
 	{
 		if (Result.IsError())
 		{
@@ -677,13 +677,13 @@ TOnlineAsyncOpHandle<FModifyLobbyJoinPolicy> FLobbiesEOSGS::ModifyLobbyJoinPolic
 		InAsyncOp.Data.Set<TSharedRef<FClientLobbyDataChanges>>(LobbyChangesKeyName, LobbyChanges);
 		LobbyChanges->JoinPolicy = Params.JoinPolicy;
 
-		FModifyLobbyDataImpl::Params ModifyLobbyDataParams;
+		FLobbiesModifyLobbyDataImpl::Params ModifyLobbyDataParams;
 		ModifyLobbyDataParams.LobbyData = LobbyData;
 		ModifyLobbyDataParams.LocalUserId = Params.LocalUserId;
 		ModifyLobbyDataParams.Changes = LobbyChanges;
 		return ModifyLobbyDataImpl(MoveTemp(ModifyLobbyDataParams));
 	})
-	.Then([this](TOnlineAsyncOp<FModifyLobbyJoinPolicy>& InAsyncOp, TDefaultErrorResult<FModifyLobbyDataImpl> Result)
+	.Then([this](TOnlineAsyncOp<FModifyLobbyJoinPolicy>& InAsyncOp, TDefaultErrorResult<FLobbiesModifyLobbyDataImpl> Result)
 	{
 		if (Result.IsError())
 		{
@@ -734,13 +734,13 @@ TOnlineAsyncOpHandle<FModifyLobbyAttributes> FLobbiesEOSGS::ModifyLobbyAttribute
 		LobbyChanges->MutatedAttributes = Params.MutatedAttributes;
 		LobbyChanges->ClearedAttributes = Params.ClearedAttributes;
 
-		FModifyLobbyDataImpl::Params ModifyLobbyDataParams;
+		FLobbiesModifyLobbyDataImpl::Params ModifyLobbyDataParams;
 		ModifyLobbyDataParams.LobbyData = LobbyDataRegistry->Find(Params.LobbyId);
 		ModifyLobbyDataParams.LocalUserId = Params.LocalUserId;
 		ModifyLobbyDataParams.Changes = LobbyChanges;
 		return ModifyLobbyDataImpl(MoveTemp(ModifyLobbyDataParams));
 	})
-	.Then([this](TOnlineAsyncOp<FModifyLobbyAttributes>& InAsyncOp, TDefaultErrorResult<FModifyLobbyDataImpl> Result)
+	.Then([this](TOnlineAsyncOp<FModifyLobbyAttributes>& InAsyncOp, TDefaultErrorResult<FLobbiesModifyLobbyDataImpl> Result)
 	{
 		if (Result.IsError())
 		{
@@ -791,13 +791,13 @@ TOnlineAsyncOpHandle<FModifyLobbyMemberAttributes> FLobbiesEOSGS::ModifyLobbyMem
 		LobbyMemberChanges->MutatedAttributes = Params.MutatedAttributes;
 		LobbyMemberChanges->ClearedAttributes = Params.ClearedAttributes;
 
-		FModifyLobbyMemberDataImpl::Params ModifyLobbyMemberDataParams;
+		FLobbiesModifyLobbyMemberDataImpl::Params ModifyLobbyMemberDataParams;
 		ModifyLobbyMemberDataParams.LobbyData = LobbyData;
 		ModifyLobbyMemberDataParams.LocalUserId = Params.LocalUserId;
 		ModifyLobbyMemberDataParams.Changes = LobbyMemberChanges;
 		return ModifyLobbyMemberDataImpl(MoveTemp(ModifyLobbyMemberDataParams));
 	})
-	.Then([this](TOnlineAsyncOp<FModifyLobbyMemberAttributes>& InAsyncOp, TDefaultErrorResult<FModifyLobbyMemberDataImpl> Result)
+	.Then([this](TOnlineAsyncOp<FModifyLobbyMemberAttributes>& InAsyncOp, TDefaultErrorResult<FLobbiesModifyLobbyMemberDataImpl> Result)
 	{
 		if (Result.IsError())
 		{
@@ -844,11 +844,11 @@ void FLobbiesEOSGS::HandleLobbyUpdated(const EOS_Lobby_LobbyUpdateReceivedCallba
 {
 	if (TSharedPtr<FLobbyDataEOS> LobbyData = LobbyDataRegistry->Find(Data->LobbyId))
 	{
-		FProcessLobbyNotificationImpl::Params Params;
+		FLobbiesProcessLobbyNotificationImpl::Params Params;
 		Params.LobbyData = LobbyData;
 
 		ProcessLobbyNotificationImplOp(MoveTemp(Params))
-		.OnComplete([LobbyData](const TOnlineResult<FProcessLobbyNotificationImpl>& Result)
+		.OnComplete([LobbyData](const TOnlineResult<FLobbiesProcessLobbyNotificationImpl>& Result)
 		{
 			if (Result.IsError())
 			{
@@ -863,12 +863,12 @@ void FLobbiesEOSGS::HandleLobbyMemberUpdated(const EOS_Lobby_LobbyMemberUpdateRe
 {
 	if (TSharedPtr<FLobbyDataEOS> LobbyData = LobbyDataRegistry->Find(Data->LobbyId))
 	{
-		FProcessLobbyNotificationImpl::Params Params;
+		FLobbiesProcessLobbyNotificationImpl::Params Params;
 		Params.LobbyData = LobbyData;
 		Params.MutatedMembers.Add(Data->TargetUserId);
 
 		ProcessLobbyNotificationImplOp(MoveTemp(Params))
-		.OnComplete([LobbyData](const TOnlineResult<FProcessLobbyNotificationImpl>& Result)
+		.OnComplete([LobbyData](const TOnlineResult<FLobbiesProcessLobbyNotificationImpl>& Result)
 		{
 			if (Result.IsError())
 			{
@@ -883,7 +883,7 @@ void FLobbiesEOSGS::HandleLobbyMemberStatusReceived(const EOS_Lobby_LobbyMemberS
 {
 	if (TSharedPtr<FLobbyDataEOS> LobbyData = LobbyDataRegistry->Find(Data->LobbyId))
 	{
-		FProcessLobbyNotificationImpl::Params Params;
+		FLobbiesProcessLobbyNotificationImpl::Params Params;
 		Params.LobbyData = LobbyData;
 
 		switch (Data->CurrentStatus)
@@ -916,7 +916,7 @@ void FLobbiesEOSGS::HandleLobbyMemberStatusReceived(const EOS_Lobby_LobbyMemberS
 		}
 
 		ProcessLobbyNotificationImplOp(MoveTemp(Params))
-		.OnComplete([LobbyData](const TOnlineResult<FProcessLobbyNotificationImpl>& Result)
+		.OnComplete([LobbyData](const TOnlineResult<FLobbiesProcessLobbyNotificationImpl>& Result)
 		{
 			if (Result.IsError())
 			{
@@ -1029,6 +1029,57 @@ void FLobbiesEOSGS::HandleJoinLobbyAccepted(const EOS_Lobby_JoinLobbyAcceptedCal
 	}
 }
 
+#if !UE_BUILD_SHIPPING
+void FLobbiesEOSGS::CheckMetadata()
+{
+	// Metadata sanity check.
+	ToLogString(FLobbiesJoinLobbyImpl::Params());
+	ToLogString(FLobbiesJoinLobbyImpl::Result());
+	ToLogString(FLobbiesJoinLobbyMemberImpl::Params());
+	ToLogString(FLobbiesJoinLobbyMemberImpl::Result());
+	ToLogString(FLobbiesLeaveLobbyImpl::Params());
+	ToLogString(FLobbiesLeaveLobbyImpl::Result());
+	ToLogString(FLobbiesDestroyLobbyImpl::Params());
+	ToLogString(FLobbiesDestroyLobbyImpl::Result());
+	ToLogString(FLobbiesInviteLobbyMemberImpl::Params());
+	ToLogString(FLobbiesInviteLobbyMemberImpl::Result());
+	ToLogString(FLobbiesDeclineLobbyInvitationImpl::Params());
+	ToLogString(FLobbiesDeclineLobbyInvitationImpl::Result());
+	ToLogString(FLobbiesKickLobbyMemberImpl::Params());
+	ToLogString(FLobbiesKickLobbyMemberImpl::Result());
+	ToLogString(FLobbiesPromoteLobbyMemberImpl::Params());
+	ToLogString(FLobbiesPromoteLobbyMemberImpl::Result());
+	ToLogString(FLobbiesModifyLobbyDataImpl::Params());
+	ToLogString(FLobbiesModifyLobbyDataImpl::Result());
+	ToLogString(FLobbiesModifyLobbyMemberDataImpl::Params());
+	ToLogString(FLobbiesModifyLobbyMemberDataImpl::Result());
+	ToLogString(FLobbiesProcessLobbyNotificationImpl::Params());
+	ToLogString(FLobbiesProcessLobbyNotificationImpl::Result());
+	Meta::VisitFields(FLobbiesJoinLobbyImpl::Params(), [](const TCHAR* Name, auto& Field) { return false; });
+	Meta::VisitFields(FLobbiesJoinLobbyImpl::Result(), [](const TCHAR* Name, auto& Field) { return false; });
+	Meta::VisitFields(FLobbiesJoinLobbyMemberImpl::Params(), [](const TCHAR* Name, auto& Field) { return false; });
+	Meta::VisitFields(FLobbiesJoinLobbyMemberImpl::Result(), [](const TCHAR* Name, auto& Field) { return false; });
+	Meta::VisitFields(FLobbiesLeaveLobbyImpl::Params(), [](const TCHAR* Name, auto& Field) { return false; });
+	Meta::VisitFields(FLobbiesLeaveLobbyImpl::Result(), [](const TCHAR* Name, auto& Field) { return false; });
+	Meta::VisitFields(FLobbiesDestroyLobbyImpl::Params(), [](const TCHAR* Name, auto& Field) { return false; });
+	Meta::VisitFields(FLobbiesDestroyLobbyImpl::Result(), [](const TCHAR* Name, auto& Field) { return false; });
+	Meta::VisitFields(FLobbiesInviteLobbyMemberImpl::Params(), [](const TCHAR* Name, auto& Field) { return false; });
+	Meta::VisitFields(FLobbiesInviteLobbyMemberImpl::Result(), [](const TCHAR* Name, auto& Field) { return false; });
+	Meta::VisitFields(FLobbiesDeclineLobbyInvitationImpl::Params(), [](const TCHAR* Name, auto& Field) { return false; });
+	Meta::VisitFields(FLobbiesDeclineLobbyInvitationImpl::Result(), [](const TCHAR* Name, auto& Field) { return false; });
+	Meta::VisitFields(FLobbiesKickLobbyMemberImpl::Params(), [](const TCHAR* Name, auto& Field) { return false; });
+	Meta::VisitFields(FLobbiesKickLobbyMemberImpl::Result(), [](const TCHAR* Name, auto& Field) { return false; });
+	Meta::VisitFields(FLobbiesPromoteLobbyMemberImpl::Params(), [](const TCHAR* Name, auto& Field) { return false; });
+	Meta::VisitFields(FLobbiesPromoteLobbyMemberImpl::Result(), [](const TCHAR* Name, auto& Field) { return false; });
+	Meta::VisitFields(FLobbiesModifyLobbyDataImpl::Params(), [](const TCHAR* Name, auto& Field) { return false; });
+	Meta::VisitFields(FLobbiesModifyLobbyDataImpl::Result(), [](const TCHAR* Name, auto& Field) { return false; });
+	Meta::VisitFields(FLobbiesModifyLobbyMemberDataImpl::Params(), [](const TCHAR* Name, auto& Field) { return false; });
+	Meta::VisitFields(FLobbiesModifyLobbyMemberDataImpl::Result(), [](const TCHAR* Name, auto& Field) { return false; });
+	Meta::VisitFields(FLobbiesProcessLobbyNotificationImpl::Params(), [](const TCHAR* Name, auto& Field) { return false; });
+	Meta::VisitFields(FLobbiesProcessLobbyNotificationImpl::Result(), [](const TCHAR* Name, auto& Field) { return false; });
+}
+#endif
+
 void FLobbiesEOSGS::RegisterHandlers()
 {
 	// Register for lobby updates.
@@ -1147,17 +1198,17 @@ TSharedPtr<FLobbyInviteDataEOS> FLobbiesEOSGS::GetActiveInvite(FOnlineAccountIdH
 	return Result ? *Result : TSharedPtr<FLobbyInviteDataEOS>();
 }
 
-TFuture<TDefaultErrorResult<FLobbiesEOSGS::FJoinLobbyImpl>> FLobbiesEOSGS::JoinLobbyImpl(FJoinLobbyImpl::Params&& Params)
+TFuture<TDefaultErrorResult<FLobbiesJoinLobbyImpl>> FLobbiesEOSGS::JoinLobbyImpl(FLobbiesJoinLobbyImpl::Params&& Params)
 {
 	// Check prerequisites.
 	if (!Services.Get<FAuthEOSGS>()->IsLoggedIn(Params.LocalUserId))
 	{
-		return MakeFulfilledPromise<TDefaultErrorResult<FJoinLobbyImpl>>(Errors::InvalidUser()).GetFuture();
+		return MakeFulfilledPromise<TDefaultErrorResult<FLobbiesJoinLobbyImpl>>(Errors::InvalidUser()).GetFuture();
 	}
 
 	if (!Params.LobbyData.IsValid())
 	{
-		return MakeFulfilledPromise<TDefaultErrorResult<FJoinLobbyImpl>>(Errors::InvalidParams()).GetFuture();
+		return MakeFulfilledPromise<TDefaultErrorResult<FLobbiesJoinLobbyImpl>>(Errors::InvalidParams()).GetFuture();
 	}
 
 	// Check whether any of the local users is already in the target lobby.
@@ -1167,22 +1218,22 @@ TFuture<TDefaultErrorResult<FLobbiesEOSGS::FJoinLobbyImpl>> FLobbiesEOSGS::JoinL
 		{
 			if (UserActiveLobbies->Find(Params.LobbyData.ToSharedRef()))
 			{
-				return MakeFulfilledPromise<TDefaultErrorResult<FJoinLobbyImpl>>(Errors::InvalidParams()).GetFuture();
+				return MakeFulfilledPromise<TDefaultErrorResult<FLobbiesJoinLobbyImpl>>(Errors::InvalidParams()).GetFuture();
 			}
 		}
 	}
 
-	TArray<TFuture<TDefaultErrorResult<FJoinLobbyMemberImpl>>> PendingMemberJoins;
+	TArray<TFuture<TDefaultErrorResult<FLobbiesJoinLobbyMemberImpl>>> PendingMemberJoins;
 	for (FJoinLobbyLocalUserData& UserData : Params.LocalUsers)
 	{
-		FJoinLobbyMemberImpl::Params JoinLobbyMemberParams;
+		FLobbiesJoinLobbyMemberImpl::Params JoinLobbyMemberParams;
 		JoinLobbyMemberParams.LobbyData = Params.LobbyData.ToSharedRef();
 		JoinLobbyMemberParams.LocalUserId = UserData.LocalUserId;
 		JoinLobbyMemberParams.Attributes = MoveTemp(UserData.Attributes);
 
-		TSharedRef<TPromise<TDefaultErrorResult<FJoinLobbyMemberImpl>>> JoinMemberPromise = MakeShared<TPromise<TDefaultErrorResult<FJoinLobbyMemberImpl>>>();
+		TSharedRef<TPromise<TDefaultErrorResult<FLobbiesJoinLobbyMemberImpl>>> JoinMemberPromise = MakeShared<TPromise<TDefaultErrorResult<FLobbiesJoinLobbyMemberImpl>>>();
 		JoinLobbyMemberImplOp(MoveTemp(JoinLobbyMemberParams))
-		.OnComplete([JoinMemberPromise](const TOnlineResult<FJoinLobbyMemberImpl>& Result)
+		.OnComplete([JoinMemberPromise](const TOnlineResult<FLobbiesJoinLobbyMemberImpl>& Result)
 		{
 			if (Result.IsOk())
 			{
@@ -1197,14 +1248,14 @@ TFuture<TDefaultErrorResult<FLobbiesEOSGS::FJoinLobbyImpl>> FLobbiesEOSGS::JoinL
 		PendingMemberJoins.Emplace(JoinMemberPromise->GetFuture());
 	}
 
-	TPromise<TDefaultErrorResult<FJoinLobbyImpl>> Promise;
-	TFuture<TDefaultErrorResult<FJoinLobbyImpl>> Future = Promise.GetFuture();
+	TPromise<TDefaultErrorResult<FLobbiesJoinLobbyImpl>> Promise;
+	TFuture<TDefaultErrorResult<FLobbiesJoinLobbyImpl>> Future = Promise.GetFuture();
 
 	WhenAll(MoveTemp(PendingMemberJoins))
-	.Then([this, Promise = MoveTemp(Promise), Params = MoveTemp(Params)](TFuture<TArray<TDefaultErrorResult<FJoinLobbyMemberImpl>>>&& Results) mutable
+	.Then([this, Promise = MoveTemp(Promise), Params = MoveTemp(Params)](TFuture<TArray<TDefaultErrorResult<FLobbiesJoinLobbyMemberImpl>>>&& Results) mutable
 	{
 		TOptional<FOnlineError> StoredError;
-		for (TDefaultErrorResult<FJoinLobbyMemberImpl>& Result : Results.Get())
+		for (TDefaultErrorResult<FLobbiesJoinLobbyMemberImpl>& Result : Results.Get())
 		{
 			if (Result.IsError())
 			{
@@ -1215,22 +1266,22 @@ TFuture<TDefaultErrorResult<FLobbiesEOSGS::FJoinLobbyImpl>> FLobbiesEOSGS::JoinL
 
 		if (StoredError)
 		{
-			TArray<TFuture<TDefaultErrorResult<FLeaveLobbyImpl>>> PendingMemberExits;
+			TArray<TFuture<TDefaultErrorResult<FLobbiesLeaveLobbyImpl>>> PendingMemberExits;
 			for (int32 MemberIndex = 0; MemberIndex < Params.LocalUsers.Num(); ++MemberIndex)
 			{
 				const FOnlineAccountIdHandle MemberId = Params.LocalUsers[MemberIndex].LocalUserId;
 
 				if (Results.Get()[MemberIndex].IsError())
 				{
-					FLeaveLobbyImpl::Params LeaveLobbyParams;
+					FLobbiesLeaveLobbyImpl::Params LeaveLobbyParams;
 					LeaveLobbyParams.LobbyData = Params.LobbyData.ToSharedRef();
 					LeaveLobbyParams.LocalUserId = MemberId;
 
-					TPromise<TDefaultErrorResult<FLeaveLobbyImpl>> LeaveMemberPromise;
+					TPromise<TDefaultErrorResult<FLobbiesLeaveLobbyImpl>> LeaveMemberPromise;
 					PendingMemberExits.Emplace(LeaveMemberPromise.GetFuture());
 
 					LeaveLobbyImpl(MoveTemp(LeaveLobbyParams))
-					.Then([LeaveMemberPromise = MoveTemp(LeaveMemberPromise)](TFuture<TDefaultErrorResult<FLeaveLobbyImpl>>&& Future) mutable
+					.Then([LeaveMemberPromise = MoveTemp(LeaveMemberPromise)](TFuture<TDefaultErrorResult<FLobbiesLeaveLobbyImpl>>&& Future) mutable
 					{
 						if (Future.Get().IsError())
 						{
@@ -1243,24 +1294,24 @@ TFuture<TDefaultErrorResult<FLobbiesEOSGS::FJoinLobbyImpl>> FLobbiesEOSGS::JoinL
 			}
 
 			WhenAll(MoveTemp(PendingMemberExits))
-			.Then([Promise = MoveTemp(Promise), StoredError = MoveTemp(*StoredError)](TFuture<TArray<TDefaultErrorResult<FLeaveLobbyImpl>>>&& Future) mutable
+			.Then([Promise = MoveTemp(Promise), StoredError = MoveTemp(*StoredError)](TFuture<TArray<TDefaultErrorResult<FLobbiesLeaveLobbyImpl>>>&& Future) mutable
 			{
 				Promise.EmplaceValue(MoveTemp(StoredError));
 			});
 		}
 		else
 		{
-			Promise.EmplaceValue(FJoinLobbyImpl::Result{});
+			Promise.EmplaceValue(FLobbiesJoinLobbyImpl::Result{});
 		}
 	});
 
 	return Future;
 }
 
-TOnlineAsyncOpHandle<FLobbiesEOSGS::FJoinLobbyMemberImpl> FLobbiesEOSGS::JoinLobbyMemberImplOp(FJoinLobbyMemberImpl::Params&& InParams)
+TOnlineAsyncOpHandle<FLobbiesJoinLobbyMemberImpl> FLobbiesEOSGS::JoinLobbyMemberImplOp(FLobbiesJoinLobbyMemberImpl::Params&& InParams)
 {
-	TOnlineAsyncOpRef<FJoinLobbyMemberImpl> Op = GetOp<FJoinLobbyMemberImpl>(MoveTemp(InParams));
-	const FJoinLobbyMemberImpl::Params& Params = Op->GetParams();
+	TOnlineAsyncOpRef<FLobbiesJoinLobbyMemberImpl> Op = GetOp<FLobbiesJoinLobbyMemberImpl>(MoveTemp(InParams));
+	const FLobbiesJoinLobbyMemberImpl::Params& Params = Op->GetParams();
 
 	// Setup lobby details - Prefer UI event before invitation before search result.
 	TSharedPtr<FLobbyDetailsEOS> LobbyDetails = Params.LobbyData->GetUserLobbyDetails(Params.LocalUserId);
@@ -1277,9 +1328,9 @@ TOnlineAsyncOpHandle<FLobbiesEOSGS::FJoinLobbyMemberImpl> FLobbiesEOSGS::JoinLob
 		return Op->GetHandle();
 	}
 
-	Op->Then([this, LobbyDetails, LobbyData = Params.LobbyData](TOnlineAsyncOp<FJoinLobbyMemberImpl>& InAsyncOp, TPromise<const EOS_Lobby_JoinLobbyCallbackInfo*>&& Promise)
+	Op->Then([this, LobbyDetails, LobbyData = Params.LobbyData](TOnlineAsyncOp<FLobbiesJoinLobbyMemberImpl>& InAsyncOp, TPromise<const EOS_Lobby_JoinLobbyCallbackInfo*>&& Promise)
 	{
-		const FLobbiesEOSGS::FJoinLobbyMemberImpl::Params& Params = InAsyncOp.GetParams();
+		const FLobbiesJoinLobbyMemberImpl::Params& Params = InAsyncOp.GetParams();
 
 		EOS_Lobby_JoinLobbyOptions JoinLobbyOptions = {};
 		JoinLobbyOptions.ApiVersion = EOS_LOBBY_JOINLOBBY_API_LATEST;
@@ -1289,7 +1340,7 @@ TOnlineAsyncOpHandle<FLobbiesEOSGS::FJoinLobbyMemberImpl> FLobbiesEOSGS::JoinLob
 		JoinLobbyOptions.LocalRTCOptions = nullptr;
 		EOS_Async(EOS_Lobby_JoinLobby, LobbyPrerequisites->LobbyInterfaceHandle, JoinLobbyOptions, MoveTemp(Promise));
 	})
-	.Then([this](TOnlineAsyncOp<FJoinLobbyMemberImpl>& InAsyncOp, const EOS_Lobby_JoinLobbyCallbackInfo* Data)
+	.Then([this](TOnlineAsyncOp<FLobbiesJoinLobbyMemberImpl>& InAsyncOp, const EOS_Lobby_JoinLobbyCallbackInfo* Data)
 	{
 		if (Data->ResultCode != EOS_EResult::EOS_Success)
 		{
@@ -1297,24 +1348,24 @@ TOnlineAsyncOpHandle<FLobbiesEOSGS::FJoinLobbyMemberImpl> FLobbiesEOSGS::JoinLob
 			InAsyncOp.SetError(Errors::Unknown(FromEOSError(Data->ResultCode)));
 		}
 	})
-	.Then([this](TOnlineAsyncOp<FJoinLobbyMemberImpl>& InAsyncOp)
+	.Then([this](TOnlineAsyncOp<FLobbiesJoinLobbyMemberImpl>& InAsyncOp)
 	{
-		const FLobbiesEOSGS::FJoinLobbyMemberImpl::Params& Params = InAsyncOp.GetParams();
-		FModifyLobbyMemberDataImpl::Params ModifyParams;
+		const FLobbiesJoinLobbyMemberImpl::Params& Params = InAsyncOp.GetParams();
+		FLobbiesModifyLobbyMemberDataImpl::Params ModifyParams;
 		ModifyParams.LobbyData = Params.LobbyData;
 		ModifyParams.LocalUserId = Params.LocalUserId;
 		ModifyParams.Changes = MakeShared<FClientLobbyMemberDataChanges>();
 		ModifyParams.Changes->MutatedAttributes = Params.Attributes;
 		return ModifyLobbyMemberDataImpl(MoveTemp(ModifyParams));
 	})
-	.Then([this](TOnlineAsyncOp<FJoinLobbyMemberImpl>& InAsyncOp, TDefaultErrorResult<FModifyLobbyMemberDataImpl>&& Result)
+	.Then([this](TOnlineAsyncOp<FLobbiesJoinLobbyMemberImpl>& InAsyncOp, TDefaultErrorResult<FLobbiesModifyLobbyMemberDataImpl>&& Result)
 	{
-		const FLobbiesEOSGS::FJoinLobbyMemberImpl::Params& Params = InAsyncOp.GetParams();
+		const FLobbiesJoinLobbyMemberImpl::Params& Params = InAsyncOp.GetParams();
 
 		if (Result.IsError())
 		{
 			// Failed to set attributes - leave the lobby.
-			FLeaveLobbyImpl::Params LeaveLobbyParams;
+			FLobbiesLeaveLobbyImpl::Params LeaveLobbyParams;
 			LeaveLobbyParams.LobbyData = Params.LobbyData;
 			LeaveLobbyParams.LocalUserId = Params.LocalUserId;
 
@@ -1322,7 +1373,7 @@ TOnlineAsyncOpHandle<FLobbiesEOSGS::FJoinLobbyMemberImpl> FLobbiesEOSGS::JoinLob
 			TFuture<void> Future = Promise.GetFuture();
 
 			LeaveLobbyImpl(MoveTemp(LeaveLobbyParams))
-			.Then([this, InAsyncOp = InAsyncOp.AsShared(), ErrorResult = MoveTemp(Result.GetErrorValue()), Promise = MoveTemp(Promise)](TFuture<TDefaultErrorResult<FLeaveLobbyImpl>>&& Future) mutable
+			.Then([this, InAsyncOp = InAsyncOp.AsShared(), ErrorResult = MoveTemp(Result.GetErrorValue()), Promise = MoveTemp(Promise)](TFuture<TDefaultErrorResult<FLobbiesLeaveLobbyImpl>>&& Future) mutable
 			{
 				if (Future.Get().IsError())
 				{
@@ -1341,26 +1392,26 @@ TOnlineAsyncOpHandle<FLobbiesEOSGS::FJoinLobbyMemberImpl> FLobbiesEOSGS::JoinLob
 			return MakeFulfilledPromise<void>().GetFuture();
 		}
 	})
-	.Then([this](TOnlineAsyncOp<FJoinLobbyMemberImpl>& InAsyncOp)
+	.Then([this](TOnlineAsyncOp<FLobbiesJoinLobbyMemberImpl>& InAsyncOp)
 	{
-		InAsyncOp.SetResult(FJoinLobbyMemberImpl::Result{});
+		InAsyncOp.SetResult(FLobbiesJoinLobbyMemberImpl::Result{});
 	})
 	.Enqueue(GetSerialQueue(InParams.LocalUserId));
 
 	return Op->GetHandle();
 }
 
-TFuture<TDefaultErrorResult<FLobbiesEOSGS::FLeaveLobbyImpl>> FLobbiesEOSGS::LeaveLobbyImpl(FLeaveLobbyImpl::Params&& Params)
+TFuture<TDefaultErrorResult<FLobbiesLeaveLobbyImpl>> FLobbiesEOSGS::LeaveLobbyImpl(FLobbiesLeaveLobbyImpl::Params&& Params)
 {
 	// Check prerequisites.
 	if (!Services.Get<FAuthEOSGS>()->IsLoggedIn(Params.LocalUserId))
 	{
-		return MakeFulfilledPromise<TDefaultErrorResult<FLeaveLobbyImpl>>(Errors::InvalidUser()).GetFuture();
+		return MakeFulfilledPromise<TDefaultErrorResult<FLobbiesLeaveLobbyImpl>>(Errors::InvalidUser()).GetFuture();
 	}
 
 	if (!Params.LobbyData.IsValid())
 	{
-		return MakeFulfilledPromise<TDefaultErrorResult<FLeaveLobbyImpl>>(Errors::InvalidParams()).GetFuture();
+		return MakeFulfilledPromise<TDefaultErrorResult<FLobbiesLeaveLobbyImpl>>(Errors::InvalidParams()).GetFuture();
 	}
 
 	EOS_Lobby_LeaveLobbyOptions LeaveLobbyOptions = {};
@@ -1368,8 +1419,8 @@ TFuture<TDefaultErrorResult<FLobbiesEOSGS::FLeaveLobbyImpl>> FLobbiesEOSGS::Leav
 	LeaveLobbyOptions.LocalUserId = GetProductUserIdChecked(Params.LocalUserId);
 	LeaveLobbyOptions.LobbyId = Params.LobbyData->GetLobbyIdEOS();
 
-	TPromise<TDefaultErrorResult<FLeaveLobbyImpl>> Promise;
-	TFuture<TDefaultErrorResult<FLeaveLobbyImpl>> Future = Promise.GetFuture();
+	TPromise<TDefaultErrorResult<FLobbiesLeaveLobbyImpl>> Promise;
+	TFuture<TDefaultErrorResult<FLobbiesLeaveLobbyImpl>> Future = Promise.GetFuture();
 
 	EOS_Async(EOS_Lobby_LeaveLobby, LobbyPrerequisites->LobbyInterfaceHandle, LeaveLobbyOptions,
 	[Promise = MoveTemp(Promise)](const EOS_Lobby_LeaveLobbyCallbackInfo* Result) mutable
@@ -1381,23 +1432,23 @@ TFuture<TDefaultErrorResult<FLobbiesEOSGS::FLeaveLobbyImpl>> FLobbiesEOSGS::Leav
 			return;
 		}
 
-		Promise.EmplaceValue(FLeaveLobbyImpl::Result{});
+		Promise.EmplaceValue(FLobbiesLeaveLobbyImpl::Result{});
 	});
 
 	return Future;
 }
 
-TFuture<TDefaultErrorResult<FLobbiesEOSGS::FDestroyLobbyImpl>> FLobbiesEOSGS::DestroyLobbyImpl(FDestroyLobbyImpl::Params&& Params)
+TFuture<TDefaultErrorResult<FLobbiesDestroyLobbyImpl>> FLobbiesEOSGS::DestroyLobbyImpl(FLobbiesDestroyLobbyImpl::Params&& Params)
 {
 	// Check prerequisites.
 	if (!Services.Get<FAuthEOSGS>()->IsLoggedIn(Params.LocalUserId))
 	{
-		return MakeFulfilledPromise<TDefaultErrorResult<FDestroyLobbyImpl>>(Errors::InvalidUser()).GetFuture();
+		return MakeFulfilledPromise<TDefaultErrorResult<FLobbiesDestroyLobbyImpl>>(Errors::InvalidUser()).GetFuture();
 	}
 
 	if (!Params.LobbyData.IsValid())
 	{
-		return MakeFulfilledPromise<TDefaultErrorResult<FDestroyLobbyImpl>>(Errors::InvalidParams()).GetFuture();
+		return MakeFulfilledPromise<TDefaultErrorResult<FLobbiesDestroyLobbyImpl>>(Errors::InvalidParams()).GetFuture();
 	}
 
 	EOS_Lobby_DestroyLobbyOptions DestroyLobbyOptions = {};
@@ -1405,8 +1456,8 @@ TFuture<TDefaultErrorResult<FLobbiesEOSGS::FDestroyLobbyImpl>> FLobbiesEOSGS::De
 	DestroyLobbyOptions.LocalUserId = GetProductUserIdChecked(Params.LocalUserId);
 	DestroyLobbyOptions.LobbyId = Params.LobbyData->GetLobbyIdEOS();
 
-	TPromise<TDefaultErrorResult<FDestroyLobbyImpl>> Promise;
-	TFuture<TDefaultErrorResult<FDestroyLobbyImpl>> Future = Promise.GetFuture();
+	TPromise<TDefaultErrorResult<FLobbiesDestroyLobbyImpl>> Promise;
+	TFuture<TDefaultErrorResult<FLobbiesDestroyLobbyImpl>> Future = Promise.GetFuture();
 
 	EOS_Async(EOS_Lobby_DestroyLobby, LobbyPrerequisites->LobbyInterfaceHandle, DestroyLobbyOptions,
 	[Promise = MoveTemp(Promise)](const EOS_Lobby_DestroyLobbyCallbackInfo* Result) mutable
@@ -1418,23 +1469,23 @@ TFuture<TDefaultErrorResult<FLobbiesEOSGS::FDestroyLobbyImpl>> FLobbiesEOSGS::De
 			return;
 		}
 
-		Promise.EmplaceValue(FDestroyLobbyImpl::Result{});
+		Promise.EmplaceValue(FLobbiesDestroyLobbyImpl::Result{});
 	});
 
 	return Future;
 }
 
-TFuture<TDefaultErrorResult<FLobbiesEOSGS::FInviteLobbyMemberImpl>> FLobbiesEOSGS::InviteLobbyMemberImpl(FInviteLobbyMemberImpl::Params&& Params)
+TFuture<TDefaultErrorResult<FLobbiesInviteLobbyMemberImpl>> FLobbiesEOSGS::InviteLobbyMemberImpl(FLobbiesInviteLobbyMemberImpl::Params&& Params)
 {
 	// Check prerequisites.
 	if (!Services.Get<FAuthEOSGS>()->IsLoggedIn(Params.LocalUserId))
 	{
-		return MakeFulfilledPromise<TDefaultErrorResult<FInviteLobbyMemberImpl>>(Errors::InvalidUser()).GetFuture();
+		return MakeFulfilledPromise<TDefaultErrorResult<FLobbiesInviteLobbyMemberImpl>>(Errors::InvalidUser()).GetFuture();
 	}
 
 	if (!Params.LobbyData.IsValid())
 	{
-		return MakeFulfilledPromise<TDefaultErrorResult<FInviteLobbyMemberImpl>>(Errors::InvalidParams()).GetFuture();
+		return MakeFulfilledPromise<TDefaultErrorResult<FLobbiesInviteLobbyMemberImpl>>(Errors::InvalidParams()).GetFuture();
 	}
 
 	EOS_Lobby_SendInviteOptions SendInviteOptions = {};
@@ -1443,8 +1494,8 @@ TFuture<TDefaultErrorResult<FLobbiesEOSGS::FInviteLobbyMemberImpl>> FLobbiesEOSG
 	SendInviteOptions.TargetUserId = GetProductUserIdChecked(Params.TargetUserId);
 	SendInviteOptions.LobbyId = Params.LobbyData->GetLobbyIdEOS();
 
-	TPromise<TDefaultErrorResult<FInviteLobbyMemberImpl>> Promise;
-	TFuture<TDefaultErrorResult<FInviteLobbyMemberImpl>> Future = Promise.GetFuture();
+	TPromise<TDefaultErrorResult<FLobbiesInviteLobbyMemberImpl>> Promise;
+	TFuture<TDefaultErrorResult<FLobbiesInviteLobbyMemberImpl>> Future = Promise.GetFuture();
 
 	EOS_Async(EOS_Lobby_SendInvite, LobbyPrerequisites->LobbyInterfaceHandle, SendInviteOptions,
 	[Promise = MoveTemp(Promise)](const EOS_Lobby_SendInviteCallbackInfo* Result) mutable
@@ -1456,18 +1507,18 @@ TFuture<TDefaultErrorResult<FLobbiesEOSGS::FInviteLobbyMemberImpl>> FLobbiesEOSG
 			return;
 		}
 
-		Promise.EmplaceValue(FInviteLobbyMemberImpl::Result{});
+		Promise.EmplaceValue(FLobbiesInviteLobbyMemberImpl::Result{});
 	});
 
 	return Future;
 }
 
-TFuture<TDefaultErrorResult<FLobbiesEOSGS::FDeclineLobbyInvitationImpl>> FLobbiesEOSGS::DeclineLobbyInvitationImpl(FDeclineLobbyInvitationImpl::Params&& Params)
+TFuture<TDefaultErrorResult<FLobbiesDeclineLobbyInvitationImpl>> FLobbiesEOSGS::DeclineLobbyInvitationImpl(FLobbiesDeclineLobbyInvitationImpl::Params&& Params)
 {
 	// Check prerequisites.
 	if (!Services.Get<FAuthEOSGS>()->IsLoggedIn(Params.LocalUserId))
 	{
-		return MakeFulfilledPromise<TDefaultErrorResult<FDeclineLobbyInvitationImpl>>(Errors::InvalidUser()).GetFuture();
+		return MakeFulfilledPromise<TDefaultErrorResult<FLobbiesDeclineLobbyInvitationImpl>>(Errors::InvalidUser()).GetFuture();
 	}
 
 	// Find the active invitation.
@@ -1475,7 +1526,7 @@ TFuture<TDefaultErrorResult<FLobbiesEOSGS::FDeclineLobbyInvitationImpl>> FLobbie
 	if (!InviteData)
 	{
 		// Todo: Errors.
-		return MakeFulfilledPromise<TDefaultErrorResult<FDeclineLobbyInvitationImpl>>(Errors::Unknown()).GetFuture();
+		return MakeFulfilledPromise<TDefaultErrorResult<FLobbiesDeclineLobbyInvitationImpl>>(Errors::Unknown()).GetFuture();
 	}
 
 	EOS_Lobby_RejectInviteOptions RejectInviteOptions = {};
@@ -1483,8 +1534,8 @@ TFuture<TDefaultErrorResult<FLobbiesEOSGS::FDeclineLobbyInvitationImpl>> FLobbie
 	RejectInviteOptions.LocalUserId = GetProductUserIdChecked(Params.LocalUserId);
 	RejectInviteOptions.InviteId = InviteData->GetInviteIdEOS();
 
-	TPromise<TDefaultErrorResult<FDeclineLobbyInvitationImpl>> Promise;
-	TFuture<TDefaultErrorResult<FDeclineLobbyInvitationImpl>> Future = Promise.GetFuture();
+	TPromise<TDefaultErrorResult<FLobbiesDeclineLobbyInvitationImpl>> Promise;
+	TFuture<TDefaultErrorResult<FLobbiesDeclineLobbyInvitationImpl>> Future = Promise.GetFuture();
 
 	EOS_Async(EOS_Lobby_RejectInvite, LobbyPrerequisites->LobbyInterfaceHandle, RejectInviteOptions,
 	[this, Promise = MoveTemp(Promise), InviteData](const EOS_Lobby_RejectInviteCallbackInfo* Result) mutable
@@ -1499,23 +1550,23 @@ TFuture<TDefaultErrorResult<FLobbiesEOSGS::FDeclineLobbyInvitationImpl>> FLobbie
 			return;
 		}
 
-		Promise.EmplaceValue(FDeclineLobbyInvitationImpl::Result{});
+		Promise.EmplaceValue(FLobbiesDeclineLobbyInvitationImpl::Result{});
 	});
 
 	return Future;
 }
 
-TFuture<TDefaultErrorResult<FLobbiesEOSGS::FKickLobbyMemberImpl>> FLobbiesEOSGS::KickLobbyMemberImpl(FKickLobbyMemberImpl::Params&& Params)
+TFuture<TDefaultErrorResult<FLobbiesKickLobbyMemberImpl>> FLobbiesEOSGS::KickLobbyMemberImpl(FLobbiesKickLobbyMemberImpl::Params&& Params)
 {
 	// Check prerequisites.
 	if (!Services.Get<FAuthEOSGS>()->IsLoggedIn(Params.LocalUserId))
 	{
-		return MakeFulfilledPromise<TDefaultErrorResult<FKickLobbyMemberImpl>>(Errors::InvalidUser()).GetFuture();
+		return MakeFulfilledPromise<TDefaultErrorResult<FLobbiesKickLobbyMemberImpl>>(Errors::InvalidUser()).GetFuture();
 	}
 
 	if (!Params.LobbyData.IsValid())
 	{
-		return MakeFulfilledPromise<TDefaultErrorResult<FKickLobbyMemberImpl>>(Errors::InvalidParams()).GetFuture();
+		return MakeFulfilledPromise<TDefaultErrorResult<FLobbiesKickLobbyMemberImpl>>(Errors::InvalidParams()).GetFuture();
 	}
 
 	// todo: check local user is lobby owner
@@ -1526,8 +1577,8 @@ TFuture<TDefaultErrorResult<FLobbiesEOSGS::FKickLobbyMemberImpl>> FLobbiesEOSGS:
 	KickMemberOptions.TargetUserId = GetProductUserIdChecked(Params.TargetUserId);
 	KickMemberOptions.LobbyId = Params.LobbyData->GetLobbyIdEOS();
 
-	TPromise<TDefaultErrorResult<FKickLobbyMemberImpl>> Promise;
-	TFuture<TDefaultErrorResult<FKickLobbyMemberImpl>> Future = Promise.GetFuture();
+	TPromise<TDefaultErrorResult<FLobbiesKickLobbyMemberImpl>> Promise;
+	TFuture<TDefaultErrorResult<FLobbiesKickLobbyMemberImpl>> Future = Promise.GetFuture();
 
 	EOS_Async(EOS_Lobby_KickMember, LobbyPrerequisites->LobbyInterfaceHandle, KickMemberOptions,
 	[Promise = MoveTemp(Promise)](const EOS_Lobby_KickMemberCallbackInfo* Result) mutable
@@ -1539,23 +1590,23 @@ TFuture<TDefaultErrorResult<FLobbiesEOSGS::FKickLobbyMemberImpl>> FLobbiesEOSGS:
 			return;
 		}
 
-		Promise.EmplaceValue(FKickLobbyMemberImpl::Result{});
+		Promise.EmplaceValue(FLobbiesKickLobbyMemberImpl::Result{});
 	});
 
 	return Future;
 }
 
-TFuture<TDefaultErrorResult<FLobbiesEOSGS::FPromoteLobbyMemberImpl>> FLobbiesEOSGS::PromoteLobbyMemberImpl(FPromoteLobbyMemberImpl::Params&& Params)
+TFuture<TDefaultErrorResult<FLobbiesPromoteLobbyMemberImpl>> FLobbiesEOSGS::PromoteLobbyMemberImpl(FLobbiesPromoteLobbyMemberImpl::Params&& Params)
 {
 	// Check prerequisites.
 	if (!Services.Get<FAuthEOSGS>()->IsLoggedIn(Params.LocalUserId))
 	{
-		return MakeFulfilledPromise<TDefaultErrorResult<FPromoteLobbyMemberImpl>>(Errors::InvalidUser()).GetFuture();
+		return MakeFulfilledPromise<TDefaultErrorResult<FLobbiesPromoteLobbyMemberImpl>>(Errors::InvalidUser()).GetFuture();
 	}
 
 	if (!Params.LobbyData.IsValid())
 	{
-		return MakeFulfilledPromise<TDefaultErrorResult<FPromoteLobbyMemberImpl>>(Errors::InvalidParams()).GetFuture();
+		return MakeFulfilledPromise<TDefaultErrorResult<FLobbiesPromoteLobbyMemberImpl>>(Errors::InvalidParams()).GetFuture();
 	}
 
 	// todo: check local user is lobby owner
@@ -1566,8 +1617,8 @@ TFuture<TDefaultErrorResult<FLobbiesEOSGS::FPromoteLobbyMemberImpl>> FLobbiesEOS
 	PromoteMemberOptions.TargetUserId = GetProductUserIdChecked(Params.TargetUserId);
 	PromoteMemberOptions.LobbyId = Params.LobbyData->GetLobbyIdEOS();
 
-	TPromise<TDefaultErrorResult<FPromoteLobbyMemberImpl>> Promise;
-	TFuture<TDefaultErrorResult<FPromoteLobbyMemberImpl>> Future = Promise.GetFuture();
+	TPromise<TDefaultErrorResult<FLobbiesPromoteLobbyMemberImpl>> Promise;
+	TFuture<TDefaultErrorResult<FLobbiesPromoteLobbyMemberImpl>> Future = Promise.GetFuture();
 
 	EOS_Async(EOS_Lobby_PromoteMember, LobbyPrerequisites->LobbyInterfaceHandle, PromoteMemberOptions,
 	[Promise = MoveTemp(Promise)](const EOS_Lobby_PromoteMemberCallbackInfo* Result) mutable
@@ -1579,41 +1630,41 @@ TFuture<TDefaultErrorResult<FLobbiesEOSGS::FPromoteLobbyMemberImpl>> FLobbiesEOS
 			return;
 		}
 
-		Promise.EmplaceValue(FPromoteLobbyMemberImpl::Result{});
+		Promise.EmplaceValue(FLobbiesPromoteLobbyMemberImpl::Result{});
 	});
 
 	return Future;
 }
 
-TFuture<TDefaultErrorResult<FLobbiesEOSGS::FModifyLobbyDataImpl>> FLobbiesEOSGS::ModifyLobbyDataImpl(FModifyLobbyDataImpl::Params&& Params)
+TFuture<TDefaultErrorResult<FLobbiesModifyLobbyDataImpl>> FLobbiesEOSGS::ModifyLobbyDataImpl(FLobbiesModifyLobbyDataImpl::Params&& Params)
 {
 	// Check prerequisites.
 	if (!Services.Get<FAuthEOSGS>()->IsLoggedIn(Params.LocalUserId))
 	{
-		return MakeFulfilledPromise<TDefaultErrorResult<FModifyLobbyDataImpl>>(Errors::InvalidUser()).GetFuture();
+		return MakeFulfilledPromise<TDefaultErrorResult<FLobbiesModifyLobbyDataImpl>>(Errors::InvalidUser()).GetFuture();
 	}
 
 	if (!Params.LobbyData.IsValid())
 	{
-		return MakeFulfilledPromise<TDefaultErrorResult<FModifyLobbyDataImpl>>(Errors::InvalidParams()).GetFuture();
+		return MakeFulfilledPromise<TDefaultErrorResult<FLobbiesModifyLobbyDataImpl>>(Errors::InvalidParams()).GetFuture();
 	}
 
 	TSharedPtr<FLobbyDetailsEOS> LobbyDetails = Params.LobbyData->GetUserLobbyDetails(Params.LocalUserId);
 	if (!LobbyDetails)
 	{
-		return MakeFulfilledPromise<TDefaultErrorResult<FModifyLobbyDataImpl>>(Errors::InvalidParams()).GetFuture();
+		return MakeFulfilledPromise<TDefaultErrorResult<FLobbiesModifyLobbyDataImpl>>(Errors::InvalidParams()).GetFuture();
 	}
 
 	if (Params.LocalUserId != Params.LobbyData->GetClientLobbyData()->GetPublicData().OwnerAccountId)
 	{
-		return MakeFulfilledPromise<TDefaultErrorResult<FModifyLobbyDataImpl>>(Errors::InvalidParams()).GetFuture();
+		return MakeFulfilledPromise<TDefaultErrorResult<FLobbiesModifyLobbyDataImpl>>(Errors::InvalidParams()).GetFuture();
 	}
 
 	UE_LOG(LogTemp, Warning, TEXT("[FLobbiesEOSGS::ModifyLobbyDataImpl] Start. Lobby: %s, Member: %s"),
 		*Params.LobbyData->GetLobbyId(), *ToLogString(Params.LocalUserId));
 
-	TPromise<TDefaultErrorResult<FModifyLobbyDataImpl>> Promise;
-	TFuture<TDefaultErrorResult<FModifyLobbyDataImpl>> Future = Promise.GetFuture();
+	TPromise<TDefaultErrorResult<FLobbiesModifyLobbyDataImpl>> Promise;
+	TFuture<TDefaultErrorResult<FLobbiesModifyLobbyDataImpl>> Future = Promise.GetFuture();
 
 	LobbyDetails->ApplyLobbyDataUpdateFromLocalChanges(Params.LocalUserId, *Params.Changes)
 	.Then([Promise = MoveTemp(Promise)](TFuture<EOS_EResult>&& Future) mutable
@@ -1628,33 +1679,33 @@ TFuture<TDefaultErrorResult<FLobbiesEOSGS::FModifyLobbyDataImpl>> FLobbiesEOSGS:
 			return;
 		}
 
-		Promise.EmplaceValue(FModifyLobbyDataImpl::Result{});
+		Promise.EmplaceValue(FLobbiesModifyLobbyDataImpl::Result{});
 	});
 
 	return Future;
 }
 
-TFuture<TDefaultErrorResult<FLobbiesEOSGS::FModifyLobbyMemberDataImpl>> FLobbiesEOSGS::ModifyLobbyMemberDataImpl(FModifyLobbyMemberDataImpl::Params&& Params)
+TFuture<TDefaultErrorResult<FLobbiesModifyLobbyMemberDataImpl>> FLobbiesEOSGS::ModifyLobbyMemberDataImpl(FLobbiesModifyLobbyMemberDataImpl::Params&& Params)
 {
 	// Check prerequisites.
 	if (!Services.Get<FAuthEOSGS>()->IsLoggedIn(Params.LocalUserId))
 	{
-		return MakeFulfilledPromise<TDefaultErrorResult<FModifyLobbyMemberDataImpl>>(Errors::InvalidUser()).GetFuture();
+		return MakeFulfilledPromise<TDefaultErrorResult<FLobbiesModifyLobbyMemberDataImpl>>(Errors::InvalidUser()).GetFuture();
 	}
 
 	if (!Params.LobbyData.IsValid())
 	{
-		return MakeFulfilledPromise<TDefaultErrorResult<FModifyLobbyMemberDataImpl>>(Errors::InvalidParams()).GetFuture();
+		return MakeFulfilledPromise<TDefaultErrorResult<FLobbiesModifyLobbyMemberDataImpl>>(Errors::InvalidParams()).GetFuture();
 	}
 
 	TSharedPtr<FLobbyDetailsEOS> LobbyDetails = Params.LobbyData->GetUserLobbyDetails(Params.LocalUserId);
 	if (!LobbyDetails)
 	{
-		return MakeFulfilledPromise<TDefaultErrorResult<FModifyLobbyMemberDataImpl>>(Errors::InvalidParams()).GetFuture();
+		return MakeFulfilledPromise<TDefaultErrorResult<FLobbiesModifyLobbyMemberDataImpl>>(Errors::InvalidParams()).GetFuture();
 	}
 
-	TPromise<TDefaultErrorResult<FModifyLobbyMemberDataImpl>> Promise;
-	TFuture<TDefaultErrorResult<FModifyLobbyMemberDataImpl>> Future = Promise.GetFuture();
+	TPromise<TDefaultErrorResult<FLobbiesModifyLobbyMemberDataImpl>> Promise;
+	TFuture<TDefaultErrorResult<FLobbiesModifyLobbyMemberDataImpl>> Future = Promise.GetFuture();
 
 	LobbyDetails->ApplyLobbyMemberDataUpdateFromLocalChanges(Params.LocalUserId, *Params.Changes)
 	.Then([Promise = MoveTemp(Promise)](TFuture<EOS_EResult>&& Future) mutable
@@ -1667,16 +1718,16 @@ TFuture<TDefaultErrorResult<FLobbiesEOSGS::FModifyLobbyMemberDataImpl>> FLobbies
 			return;
 		}
 
-		Promise.EmplaceValue(FModifyLobbyMemberDataImpl::Result{});
+		Promise.EmplaceValue(FLobbiesModifyLobbyMemberDataImpl::Result{});
 	});
 
 	return Future;
 }
 
-TOnlineAsyncOpHandle<FLobbiesEOSGS::FProcessLobbyNotificationImpl> FLobbiesEOSGS::ProcessLobbyNotificationImplOp(FProcessLobbyNotificationImpl::Params&& InParams)
+TOnlineAsyncOpHandle<FLobbiesProcessLobbyNotificationImpl> FLobbiesEOSGS::ProcessLobbyNotificationImplOp(FLobbiesProcessLobbyNotificationImpl::Params&& InParams)
 {
-	TOnlineAsyncOpRef<FProcessLobbyNotificationImpl> Op = GetOp<FProcessLobbyNotificationImpl>(MoveTemp(InParams));
-	const FProcessLobbyNotificationImpl::Params& Params = Op->GetParams();
+	TOnlineAsyncOpRef<FLobbiesProcessLobbyNotificationImpl> Op = GetOp<FLobbiesProcessLobbyNotificationImpl>(MoveTemp(InParams));
+	const FLobbiesProcessLobbyNotificationImpl::Params& Params = Op->GetParams();
 
 	if (!Params.LobbyData.IsValid())
 	{
@@ -1684,9 +1735,9 @@ TOnlineAsyncOpHandle<FLobbiesEOSGS::FProcessLobbyNotificationImpl> FLobbiesEOSGS
 		return Op->GetHandle();
 	}
 
-	Op->Then([this](TOnlineAsyncOp<FProcessLobbyNotificationImpl>& InAsyncOp)
+	Op->Then([this](TOnlineAsyncOp<FLobbiesProcessLobbyNotificationImpl>& InAsyncOp)
 	{
-		const FProcessLobbyNotificationImpl::Params& Params = InAsyncOp.GetParams();
+		const FLobbiesProcessLobbyNotificationImpl::Params& Params = InAsyncOp.GetParams();
 
 		// Notifications do not always indicate a user. Try to find a valid lobby details object to
 		// handle acquiring data snapshots.
@@ -1704,9 +1755,9 @@ TOnlineAsyncOpHandle<FLobbiesEOSGS::FProcessLobbyNotificationImpl> FLobbiesEOSGS
 		// Fetch lobby snapshot. Fetching the snapshot resolves the account ids of all lobby members in the snapshot.
 		return LobbyDetails->GetLobbySnapshot();
 	})
-	.Then([this](TOnlineAsyncOp<FProcessLobbyNotificationImpl>& InAsyncOp, TDefaultErrorResultInternal<TSharedRef<FClientLobbySnapshot>>&& LobbySnapshotResult)
+	.Then([this](TOnlineAsyncOp<FLobbiesProcessLobbyNotificationImpl>& InAsyncOp, TDefaultErrorResultInternal<TSharedRef<FClientLobbySnapshot>>&& LobbySnapshotResult)
 	{
-		const FProcessLobbyNotificationImpl::Params& Params = InAsyncOp.GetParams();
+		const FLobbiesProcessLobbyNotificationImpl::Params& Params = InAsyncOp.GetParams();
 		TSharedRef<FLobbyDetailsEOS> LobbyDetails = GetOpDataChecked<TSharedRef<FLobbyDetailsEOS>>(InAsyncOp, LobbyDetailsKeyName);
 
 		if (LobbySnapshotResult.IsError())
@@ -1760,7 +1811,7 @@ TOnlineAsyncOpHandle<FLobbiesEOSGS::FProcessLobbyNotificationImpl> FLobbiesEOSGS
 			RemoveActiveLobby(LeavingMember, Params.LobbyData.ToSharedRef());
 		}
 
-		InAsyncOp.SetResult(FProcessLobbyNotificationImpl::Result{});
+		InAsyncOp.SetResult(FLobbiesProcessLobbyNotificationImpl::Result{});
 	})
 	.Enqueue(GetSerialQueue());
 
diff --git a/Engine/Plugins/Online/Experimental/OnlineServicesEOSGS/Source/Private/Online/LobbiesEOSGSTypes.cpp b/Engine/Plugins/Online/Experimental/OnlineServicesEOSGS/Source/Private/Online/LobbiesEOSGSTypes.cpp
--- a/Engine/Plugins/Online/Experimental/OnlineServicesEOSGS/Source/Private/Online/LobbiesEOSGSTypes.cpp
+++ b/Engine/Plugins/Online/Experimental/OnlineServicesEOSGS/Source/Private/Online/LobbiesEOSGSTypes.cpp
@@ -1033,4 +1033,9 @@ FLobbySearchEOS::FLobbySearchEOS(const TSharedRef<FSearchHandle>& SearchHandle,
 {
 }
 
+FString ToLogString(const FLobbyDataEOS& LobbyData)
+{
+	return FString::Printf(TEXT("[%s:%s]"), *ToLogString(LobbyData.GetLobbyIdHandle()), *LobbyData.GetLobbyId());
+}
+
 /* UE::Online */ }
\ No newline at end of file
diff --git a/Engine/Plugins/Online/Experimental/OnlineServicesEOSGS/Source/Private/Online/LobbiesEOSGSTypes.h b/Engine/Plugins/Online/Experimental/OnlineServicesEOSGS/Source/Private/Online/LobbiesEOSGSTypes.h
--- a/Engine/Plugins/Online/Experimental/OnlineServicesEOSGS/Source/Private/Online/LobbiesEOSGSTypes.h
+++ b/Engine/Plugins/Online/Experimental/OnlineServicesEOSGS/Source/Private/Online/LobbiesEOSGSTypes.h
@@ -485,4 +485,11 @@ private:
 	TArray<TSharedRef<FLobbyDataEOS>> Lobbies;
 };
 
+FString ToLogString(const FLobbyDataEOS& LobbyData);
+
+inline FString LexToString(const FLobbyDataEOS& LobbyData)
+{
+	return ToLogString(LobbyData);
+}
+
 /* UE::Online */ }
diff --git a/Engine/Plugins/Online/Experimental/OnlineServicesEOSGS/Source/Public/Online/AuthEOSGS.h b/Engine/Plugins/Online/Experimental/OnlineServicesEOSGS/Source/Public/Online/AuthEOSGS.h
--- a/Engine/Plugins/Online/Experimental/OnlineServicesEOSGS/Source/Public/Online/AuthEOSGS.h
+++ b/Engine/Plugins/Online/Experimental/OnlineServicesEOSGS/Source/Public/Online/AuthEOSGS.h
@@ -21,6 +21,139 @@ class FOnlineServicesEOSGS;
 struct FAccountInfoEOS;
 class FAccountInfoRegistryEOS;
 
+
+struct FAuthLoginEASImpl
+{
+	static constexpr TCHAR Name[] = TEXT("LoginEASImpl");
+
+	struct Params
+	{
+		FPlatformUserId PlatformUserId = PLATFORMUSERID_NONE;
+		FName CredentialsType;
+		FString CredentialsId;
+		TVariant<FString, FExternalAuthToken> CredentialsToken;
+		TArray<FString> Scopes;
+	};
+
+	struct Result
+	{
+		EOS_EpicAccountId EpicAccountId = nullptr;
+	};
+};
+
+struct FAuthLogoutEASImpl
+{
+	static constexpr TCHAR Name[] = TEXT("LogoutEASImpl");
+
+	struct Params
+	{
+		EOS_EpicAccountId EpicAccountId = nullptr;
+	};
+
+	struct Result
+	{
+	};
+};
+
+struct FAuthGetExternalAuthTokenImpl
+{
+	static constexpr TCHAR Name[] = TEXT("GetExternalAuthTokenImpl");
+
+	struct Params
+	{
+		EOS_EpicAccountId EpicAccountId = nullptr;
+	};
+
+	struct Result
+	{
+		FExternalAuthToken Token;
+	};
+};
+
+struct FAuthLoginConnectImpl
+{
+	static constexpr TCHAR Name[] = TEXT("LoginConnectImpl");
+
+	struct Params
+	{
+		FPlatformUserId PlatformUserId = PLATFORMUSERID_NONE;
+		FExternalAuthToken ExternalAuthToken;
+	};
+
+	struct Result
+	{
+		EOS_ProductUserId ProductUserId = nullptr;
+	};
+};
+
+struct FAuthConnectLoginRecoveryImpl
+{
+	static constexpr TCHAR Name[] = TEXT("ConnectLoginRecovery");
+
+	struct Params
+	{
+		/** The Epic Account ID of the local user whose connect login should be recovered. */
+		EOS_EpicAccountId LocalUserId = nullptr;
+	};
+
+	struct Result
+	{
+	};
+};
+
+struct FAuthHandleConnectLoginStatusChangedImpl
+{
+	static constexpr TCHAR Name[] = TEXT("HandleConnectLoginStatusChangedImpl");
+
+	struct Params
+	{
+		/** The Product User ID of the local player whose status has changed. */
+		EOS_ProductUserId LocalUserId = nullptr;
+		/** The status prior to the change. */
+		EOS_ELoginStatus PreviousStatus = EOS_ELoginStatus::EOS_LS_NotLoggedIn;
+		/** The status at the time of the notification. */
+		EOS_ELoginStatus CurrentStatus = EOS_ELoginStatus::EOS_LS_NotLoggedIn;
+	};
+
+	struct Result
+	{
+	};
+};
+
+struct FAuthHandleConnectAuthNotifyExpirationImpl
+{
+	static constexpr TCHAR Name[] = TEXT("HandleConnectAuthNotifyExpirationImpl");
+
+	struct Params
+	{
+		/** The Product User ID of the local player whose status has changed. */
+		EOS_ProductUserId LocalUserId = nullptr;
+	};
+
+	struct Result
+	{
+	};
+};
+
+struct FAuthHandleEASLoginStatusChangedImpl
+{
+	static constexpr TCHAR Name[] = TEXT("HandleEASLoginStatusChangedImpl");
+
+	struct Params
+	{
+		/** The Epic Account ID of the local user whose status has changed */
+		EOS_EpicAccountId LocalUserId = nullptr;
+		/** The status prior to the change */
+		EOS_ELoginStatus PreviousStatus = EOS_ELoginStatus::EOS_LS_NotLoggedIn;
+		/** The status at the time of the notification */
+		EOS_ELoginStatus CurrentStatus = EOS_ELoginStatus::EOS_LS_NotLoggedIn;
+	};
+
+	struct Result
+	{
+	};
+};
+
 struct FAccountInfoEOS final : public FAccountInfo
 {
 	FTSTicker::FDelegateHandle RestoreLoginTimer;
@@ -84,153 +217,14 @@ public:
 protected:
 	// internal operations.
 
-	struct FLoginEASImpl
-	{
-		static constexpr TCHAR Name[] = TEXT("LoginEASImpl");
-
-		struct Params
-		{
-			FPlatformUserId PlatformUserId = PLATFORMUSERID_NONE;
-			FName CredentialsType;
-			FString CredentialsId;
-			TVariant<FString, FExternalAuthToken> CredentialsToken;
-			TArray<FString> Scopes;
-		};
-
-		struct Result
-		{
-			EOS_EpicAccountId EpicAccountId = nullptr;
-		};
-	};
-
-	TFuture<TDefaultErrorResult<FLoginEASImpl>> LoginEASImpl(const FLoginEASImpl::Params& LoginParams);
-
-	struct FLogoutEASImpl
-	{
-		static constexpr TCHAR Name[] = TEXT("LogoutEASImpl");
-
-		struct Params
-		{
-			EOS_EpicAccountId EpicAccountId = nullptr;
-		};
-
-		struct Result
-		{
-		};
-	};
-
-	TFuture<TDefaultErrorResult<FLogoutEASImpl>> LogoutEASImpl(const FLogoutEASImpl::Params& LogoutParams);
-
-	struct FGetExternalAuthTokenImpl
-	{
-		static constexpr TCHAR Name[] = TEXT("GetExternalAuthTokenImpl");
-
-		struct Params
-		{
-			EOS_EpicAccountId EpicAccountId = nullptr;
-		};
-
-		struct Result
-		{
-			FExternalAuthToken Token;
-		};
-	};
-
-	TDefaultErrorResult<FGetExternalAuthTokenImpl> GetExternalAuthTokenImpl(const FGetExternalAuthTokenImpl::Params& Params);
-
-	struct FLoginConnectImpl
-	{
-		static constexpr TCHAR Name[] = TEXT("LoginConnectImpl");
-
-		struct Params
-		{
-			FPlatformUserId PlatformUserId = PLATFORMUSERID_NONE;
-			FExternalAuthToken ExternalAuthToken;
-		};
-
-		struct Result
-		{
-			EOS_ProductUserId ProductUserId = nullptr;
-		};
-	};
-
-	TFuture<TDefaultErrorResult<FLoginConnectImpl>> LoginConnectImpl(const FLoginConnectImpl::Params& LoginParams);
-
-	struct FConnectLoginRecoveryImpl
-	{
-		static constexpr TCHAR Name[] = TEXT("ConnectLoginRecovery");
-
-		struct Params
-		{
-			/** The Epic Account ID of the local user whose connect login should be recovered. */
-			EOS_EpicAccountId LocalUserId = nullptr;
-		};
-
-		struct Result
-		{
-		};
-	};
-
-	TOnlineAsyncOpHandle<FConnectLoginRecoveryImpl> ConnectLoginRecoveryImplOp(FConnectLoginRecoveryImpl::Params&& Params);
-
-	struct FHandleConnectLoginStatusChangedImpl
-	{
-		static constexpr TCHAR Name[] = TEXT("HandleConnectLoginStatusChangedImpl");
-
-		struct Params
-		{
-			/** The Product User ID of the local player whose status has changed. */
-			EOS_ProductUserId LocalUserId = nullptr;
-			/** The status prior to the change. */
-			EOS_ELoginStatus PreviousStatus = EOS_ELoginStatus::EOS_LS_NotLoggedIn;
-			/** The status at the time of the notification. */
-			EOS_ELoginStatus CurrentStatus = EOS_ELoginStatus::EOS_LS_NotLoggedIn;
-		};
-
-		struct Result
-		{
-		};
-	};
-
-	TOnlineAsyncOpHandle<FHandleConnectLoginStatusChangedImpl> HandleConnectLoginStatusChangedImplOp(FHandleConnectLoginStatusChangedImpl::Params&& Params);
-
-	struct FHandleConnectAuthNotifyExpirationImpl
-	{
-		static constexpr TCHAR Name[] = TEXT("HandleConnectAuthNotifyExpirationImpl");
-
-		struct Params
-		{
-			/** The Product User ID of the local player whose status has changed. */
-			EOS_ProductUserId LocalUserId = nullptr;
-		};
-
-		struct Result
-		{
-		};
-	};
-
-	TOnlineAsyncOpHandle<FHandleConnectAuthNotifyExpirationImpl> HandleConnectAuthNotifyExpirationImplOp(FHandleConnectAuthNotifyExpirationImpl::Params&& Params);
-
-	struct FHandleEASLoginStatusChangedImpl
-	{
-		static constexpr TCHAR Name[] = TEXT("HandleEASLoginStatusChangedImpl");
-
-		struct Params
-		{
-			/** The Epic Account ID of the local user whose status has changed */
-			EOS_EpicAccountId LocalUserId = nullptr;
-			/** The status prior to the change */
-			EOS_ELoginStatus PrevStatus = EOS_ELoginStatus::EOS_LS_NotLoggedIn;
-			/** The status at the time of the notification */
-			EOS_ELoginStatus CurrentStatus = EOS_ELoginStatus::EOS_LS_NotLoggedIn;
-		};
-
-		struct Result
-		{
-		};
-	};
-
-	TOnlineAsyncOpHandle<FHandleEASLoginStatusChangedImpl> HandleEASLoginStatusChangedImplOp(FHandleEASLoginStatusChangedImpl::Params&& Params);
+	TFuture<TDefaultErrorResult<FAuthLoginEASImpl>> LoginEASImpl(const FAuthLoginEASImpl::Params& LoginParams);
+	TFuture<TDefaultErrorResult<FAuthLogoutEASImpl>> LogoutEASImpl(const FAuthLogoutEASImpl::Params& LogoutParams);
+	TDefaultErrorResult<FAuthGetExternalAuthTokenImpl> GetExternalAuthTokenImpl(const FAuthGetExternalAuthTokenImpl::Params& Params);
+	TFuture<TDefaultErrorResult<FAuthLoginConnectImpl>> LoginConnectImpl(const FAuthLoginConnectImpl::Params& LoginParams);
+	TOnlineAsyncOpHandle<FAuthConnectLoginRecoveryImpl> ConnectLoginRecoveryImplOp(FAuthConnectLoginRecoveryImpl::Params&& Params);
+	TOnlineAsyncOpHandle<FAuthHandleConnectLoginStatusChangedImpl> HandleConnectLoginStatusChangedImplOp(FAuthHandleConnectLoginStatusChangedImpl::Params&& Params);
+	TOnlineAsyncOpHandle<FAuthHandleConnectAuthNotifyExpirationImpl> HandleConnectAuthNotifyExpirationImplOp(FAuthHandleConnectAuthNotifyExpirationImpl::Params&& Params);
+	TOnlineAsyncOpHandle<FAuthHandleEASLoginStatusChangedImpl> HandleEASLoginStatusChangedImplOp(FAuthHandleEASLoginStatusChangedImpl::Params&& Params);
 
 protected:
 	// Service event handling.
@@ -241,6 +235,10 @@ protected:
 	void OnEASLoginStatusChanged(const EOS_Auth_LoginStatusChangedCallbackInfo* Data);
 
 protected:
+#if !UE_BUILD_SHIPPING
+	static void CheckMetadata();
+#endif
+
 	virtual const FAccountInfoRegistry& GetAccountInfoRegistry() const override;
 
 	void InitializeConnectLoginRecoveryTimer(const TSharedRef<FAccountInfoEOS>& UserAuthData);
@@ -254,4 +252,85 @@ protected:
 	FAccountInfoRegistryEOS AccountInfoRegistryEOS;
 };
 
+namespace Meta {
+
+BEGIN_ONLINE_STRUCT_META(FAuthLoginEASImpl::Params)
+	ONLINE_STRUCT_FIELD(FAuthLoginEASImpl::Params, PlatformUserId),
+	ONLINE_STRUCT_FIELD(FAuthLoginEASImpl::Params, CredentialsType),
+	ONLINE_STRUCT_FIELD(FAuthLoginEASImpl::Params, CredentialsId),
+	ONLINE_STRUCT_FIELD(FAuthLoginEASImpl::Params, CredentialsToken),
+	ONLINE_STRUCT_FIELD(FAuthLoginEASImpl::Params, Scopes)
+END_ONLINE_STRUCT_META()
+
+BEGIN_ONLINE_STRUCT_META(FAuthLoginEASImpl::Result)
+	ONLINE_STRUCT_FIELD(FAuthLoginEASImpl::Result, EpicAccountId)
+END_ONLINE_STRUCT_META()
+
+BEGIN_ONLINE_STRUCT_META(FAuthLogoutEASImpl::Params)
+	ONLINE_STRUCT_FIELD(FAuthLogoutEASImpl::Params, EpicAccountId)
+END_ONLINE_STRUCT_META()
+
+BEGIN_ONLINE_STRUCT_META(FAuthLogoutEASImpl::Result)
+END_ONLINE_STRUCT_META()
+
+BEGIN_ONLINE_STRUCT_META(FAuthGetExternalAuthTokenImpl::Params)
+	ONLINE_STRUCT_FIELD(FAuthGetExternalAuthTokenImpl::Params, EpicAccountId)
+END_ONLINE_STRUCT_META()
+
+BEGIN_ONLINE_STRUCT_META(FAuthGetExternalAuthTokenImpl::Result)
+	ONLINE_STRUCT_FIELD(FAuthGetExternalAuthTokenImpl::Result, Token)
+END_ONLINE_STRUCT_META()
+
+BEGIN_ONLINE_STRUCT_META(FAuthLoginConnectImpl::Params)
+	ONLINE_STRUCT_FIELD(FAuthLoginConnectImpl::Params, PlatformUserId),
+	ONLINE_STRUCT_FIELD(FAuthLoginConnectImpl::Params, ExternalAuthToken)
+END_ONLINE_STRUCT_META()
+
+BEGIN_ONLINE_STRUCT_META(FAuthLoginConnectImpl::Result)
+	ONLINE_STRUCT_FIELD(FAuthLoginConnectImpl::Result, ProductUserId)
+END_ONLINE_STRUCT_META()
+
+BEGIN_ONLINE_STRUCT_META(FAuthConnectLoginRecoveryImpl::Params)
+	ONLINE_STRUCT_FIELD(FAuthConnectLoginRecoveryImpl::Params, LocalUserId)
+END_ONLINE_STRUCT_META()
+
+BEGIN_ONLINE_STRUCT_META(FAuthConnectLoginRecoveryImpl::Result)
+END_ONLINE_STRUCT_META()
+
+BEGIN_ONLINE_STRUCT_META(FAuthHandleConnectLoginStatusChangedImpl::Params)
+	ONLINE_STRUCT_FIELD(FAuthHandleConnectLoginStatusChangedImpl::Params, LocalUserId),
+	ONLINE_STRUCT_FIELD(FAuthHandleConnectLoginStatusChangedImpl::Params, PreviousStatus),
+	ONLINE_STRUCT_FIELD(FAuthHandleConnectLoginStatusChangedImpl::Params, CurrentStatus)
+END_ONLINE_STRUCT_META()
+
+BEGIN_ONLINE_STRUCT_META(FAuthHandleConnectLoginStatusChangedImpl::Result)
+END_ONLINE_STRUCT_META()
+
+BEGIN_ONLINE_STRUCT_META(FAuthHandleConnectAuthNotifyExpirationImpl::Params)
+	ONLINE_STRUCT_FIELD(FAuthHandleConnectAuthNotifyExpirationImpl::Params, LocalUserId)
+END_ONLINE_STRUCT_META()
+
+BEGIN_ONLINE_STRUCT_META(FAuthHandleConnectAuthNotifyExpirationImpl::Result)
+END_ONLINE_STRUCT_META()
+
+BEGIN_ONLINE_STRUCT_META(FAuthHandleEASLoginStatusChangedImpl::Params)
+	ONLINE_STRUCT_FIELD(FAuthHandleEASLoginStatusChangedImpl::Params, LocalUserId),
+	ONLINE_STRUCT_FIELD(FAuthHandleEASLoginStatusChangedImpl::Params, PreviousStatus),
+	ONLINE_STRUCT_FIELD(FAuthHandleEASLoginStatusChangedImpl::Params, CurrentStatus)
+END_ONLINE_STRUCT_META()
+
+BEGIN_ONLINE_STRUCT_META(FAuthHandleEASLoginStatusChangedImpl::Result)
+END_ONLINE_STRUCT_META()
+
+BEGIN_ONLINE_STRUCT_META(FAccountInfoEOS)
+	ONLINE_STRUCT_FIELD(FAccountInfoEOS, AccountId),
+	ONLINE_STRUCT_FIELD(FAccountInfoEOS, PlatformUserId),
+	ONLINE_STRUCT_FIELD(FAccountInfoEOS, LoginStatus),
+	ONLINE_STRUCT_FIELD(FAccountInfoEOS, Attributes),
+	ONLINE_STRUCT_FIELD(FAccountInfoEOS, EpicAccountId),
+	ONLINE_STRUCT_FIELD(FAccountInfoEOS, ProductUserId)
+END_ONLINE_STRUCT_META()
+
+/* Meta*/ }
+
 /* UE::Online */ }
diff --git a/Engine/Plugins/Online/Experimental/OnlineServicesEOSGS/Source/Public/Online/LobbiesEOSGS.h b/Engine/Plugins/Online/Experimental/OnlineServicesEOSGS/Source/Public/Online/LobbiesEOSGS.h
--- a/Engine/Plugins/Online/Experimental/OnlineServicesEOSGS/Source/Public/Online/LobbiesEOSGS.h
+++ b/Engine/Plugins/Online/Experimental/OnlineServicesEOSGS/Source/Public/Online/LobbiesEOSGS.h
@@ -21,305 +21,299 @@ class FOnlineServicesEOSGS;
 class FLobbySearchEOS;
 struct FClientLobbyDataChanges;
 struct FClientLobbyMemberDataChanges;
-	
-class FLobbiesEOSGS : public FLobbiesCommon
-{
-public:
-	FLobbiesEOSGS(FOnlineServicesEOSGS& InServices);
-
-	virtual void Initialize() override;
-	virtual void PreShutdown() override;
 
-	virtual TOnlineAsyncOpHandle<FCreateLobby> CreateLobby(FCreateLobby::Params&& Params) override;
-	virtual TOnlineAsyncOpHandle<FFindLobbies> FindLobbies(FFindLobbies::Params&& Params) override;
-	virtual TOnlineAsyncOpHandle<FJoinLobby> JoinLobby(FJoinLobby::Params&& Params) override;
-	virtual TOnlineAsyncOpHandle<FLeaveLobby> LeaveLobby(FLeaveLobby::Params&& Params) override;
-	virtual TOnlineAsyncOpHandle<FInviteLobbyMember> InviteLobbyMember(FInviteLobbyMember::Params&& Params) override;
-	virtual TOnlineAsyncOpHandle<FDeclineLobbyInvitation> DeclineLobbyInvitation(FDeclineLobbyInvitation::Params&& Params) override;
-	virtual TOnlineAsyncOpHandle<FKickLobbyMember> KickLobbyMember(FKickLobbyMember::Params&& Params) override;
-	virtual TOnlineAsyncOpHandle<FPromoteLobbyMember> PromoteLobbyMember(FPromoteLobbyMember::Params&& Params) override;
-	virtual TOnlineAsyncOpHandle<FModifyLobbyJoinPolicy> ModifyLobbyJoinPolicy(FModifyLobbyJoinPolicy::Params&& Params) override;
-	virtual TOnlineAsyncOpHandle<FModifyLobbyAttributes> ModifyLobbyAttributes(FModifyLobbyAttributes::Params&& Params) override;
-	virtual TOnlineAsyncOpHandle<FModifyLobbyMemberAttributes> ModifyLobbyMemberAttributes(FModifyLobbyMemberAttributes::Params&& Params) override;
-	virtual TOnlineResult<FGetJoinedLobbies> GetJoinedLobbies(FGetJoinedLobbies::Params&& Params) override;
+struct FLobbiesJoinLobbyImpl
+{
+	static constexpr TCHAR Name[] = TEXT("JoinLobbyImpl");
 
-private:
-	void HandleLobbyUpdated(const EOS_Lobby_LobbyUpdateReceivedCallbackInfo* Data);
-	void HandleLobbyMemberUpdated(const EOS_Lobby_LobbyMemberUpdateReceivedCallbackInfo* Data);
-	void HandleLobbyMemberStatusReceived(const EOS_Lobby_LobbyMemberStatusReceivedCallbackInfo* Data);
-	void HandleLobbyInviteReceived(const EOS_Lobby_LobbyInviteReceivedCallbackInfo* Data);
-	void HandleLobbyInviteAccepted(const EOS_Lobby_LobbyInviteAcceptedCallbackInfo* Data);
-	void HandleJoinLobbyAccepted(const EOS_Lobby_JoinLobbyAcceptedCallbackInfo* Data);
+	struct Params
+	{
+		// The lobby handle data.
+		TSharedPtr<FLobbyDataEOS> LobbyData;
 
-protected:
-	void RegisterHandlers();
-	void UnregisterHandlers();
+		// The local user agent which will perform the action.
+		FOnlineAccountIdHandle LocalUserId;
 
-	void AddActiveLobby(FOnlineAccountIdHandle LocalUserId, const TSharedRef<FLobbyDataEOS>& LobbyData);
-	void RemoveActiveLobby(FOnlineAccountIdHandle LocalUserId, const TSharedRef<FLobbyDataEOS>& LobbyData);
+		// The local name for the lobby.
+		FName LocalName;
 
-	// Todo: store list of invites per lobby.
-	void AddActiveInvite(const TSharedRef<FLobbyInviteDataEOS>& Invite);
-	void RemoveActiveInvite(const TSharedRef<FLobbyInviteDataEOS>& Invite);
-	TSharedPtr<FLobbyInviteDataEOS> GetActiveInvite(FOnlineAccountIdHandle TargetUser, FOnlineLobbyIdHandle TargetLobbyId);
+		// Local users who will be joining the lobby.
+		TArray<FJoinLobbyLocalUserData> LocalUsers;
+	};
 
-	struct FJoinLobbyImpl
+	struct Result
 	{
-		static constexpr TCHAR Name[] = TEXT("JoinLobbyImpl");
-
-		struct Params
-		{
-			// The lobby handle data.
-			TSharedPtr<FLobbyDataEOS> LobbyData;
+	};
+};
 
-			// The local user agent which will perform the action.
-			FOnlineAccountIdHandle LocalUserId;
+struct FLobbiesJoinLobbyMemberImpl
+{
+	static constexpr TCHAR Name[] = TEXT("JoinLobbyMemberImpl");
 
-			// The local name for the lobby.
-			FName LocalName;
+	struct Params
+	{
+		// The lobby handle data.
+		TSharedPtr<FLobbyDataEOS> LobbyData;
 
-			// Local users who will be joining the lobby.
-			TArray<FJoinLobbyLocalUserData> LocalUsers;
-		};
+		// The local user agent which will perform the action.
+		FOnlineAccountIdHandle LocalUserId;
 
-		struct Result
-		{
-		};
+		// Initial attributes.
+		TMap<FLobbyAttributeId, FLobbyVariant> Attributes;
 	};
 
-	// LobbyData will be fetched from the operation data if not set in Params.
-	TFuture<TDefaultErrorResult<FJoinLobbyImpl>> JoinLobbyImpl(FJoinLobbyImpl::Params&& Params);
-
-	struct FJoinLobbyMemberImpl
+	struct Result
 	{
-		static constexpr TCHAR Name[] = TEXT("JoinLobbyMemberImpl");
+	};
+};
 
-		struct Params
-		{
-			// The lobby handle data.
-			TSharedPtr<FLobbyDataEOS> LobbyData;
+struct FLobbiesLeaveLobbyImpl
+{
+	static constexpr TCHAR Name[] = TEXT("LeaveLobbyImpl");
 
-			// The local user agent which will perform the action.
-			FOnlineAccountIdHandle LocalUserId;
+	struct Params
+	{
+		// The lobby handle data.
+		TSharedPtr<FLobbyDataEOS> LobbyData;
 
-			// Initial attributes.
-			TMap<FLobbyAttributeId, FLobbyVariant> Attributes;
-		};
+		// The local user agent which will perform the action.
+		FOnlineAccountIdHandle LocalUserId;
+	};
 
-		struct Result
-		{
-		};
+	struct Result
+	{
 	};
+};
 
-	TOnlineAsyncOpHandle<FJoinLobbyMemberImpl> JoinLobbyMemberImplOp(FJoinLobbyMemberImpl::Params&& Params);
-	TFuture<TDefaultErrorResult<FJoinLobbyMemberImpl>> JoinLobbyMemberImpl(FJoinLobbyMemberImpl::Params&& Params);
+struct FLobbiesDestroyLobbyImpl
+{
+	static constexpr TCHAR Name[] = TEXT("DestroyLobbyImpl");
 
-	struct FLeaveLobbyImpl
+	struct Params
 	{
-		static constexpr TCHAR Name[] = TEXT("LeaveLobbyImpl");
-
-		struct Params
-		{
-			// The lobby handle data.
-			TSharedPtr<FLobbyDataEOS> LobbyData;
+		// The lobby handle data.
+		TSharedPtr<FLobbyDataEOS> LobbyData;
 
-			// The local user agent which will perform the action.
-			FOnlineAccountIdHandle LocalUserId;
-		};
+		// The local user agent which will perform the action.
+		FOnlineAccountIdHandle LocalUserId;
+	};
 
-		struct Result
-		{
-		};
+	struct Result
+	{
 	};
+};
 
-	TFuture<TDefaultErrorResult<FLeaveLobbyImpl>> LeaveLobbyImpl(FLeaveLobbyImpl::Params&& Params);
+struct FLobbiesInviteLobbyMemberImpl
+{
+	static constexpr TCHAR Name[] = TEXT("InviteLobbyMemberImpl");
 
-	struct FDestroyLobbyImpl
+	struct Params
 	{
-		static constexpr TCHAR Name[] = TEXT("DestroyLobbyImpl");
+		// The lobby handle data.
+		TSharedPtr<FLobbyDataEOS> LobbyData;
 
-		struct Params
-		{
-			// The lobby handle data.
-			TSharedPtr<FLobbyDataEOS> LobbyData;
+		// The local user agent which will perform the action.
+		FOnlineAccountIdHandle LocalUserId;
 
-			// The local user agent which will perform the action.
-			FOnlineAccountIdHandle LocalUserId;
-		};
-
-		struct Result
-		{
-		};
+		// The target user for the invitation.
+		FOnlineAccountIdHandle TargetUserId;
 	};
 
-	TFuture<TDefaultErrorResult<FDestroyLobbyImpl>> DestroyLobbyImpl(FDestroyLobbyImpl::Params&& Params);
-
-	struct FInviteLobbyMemberImpl
+	struct Result
 	{
-		static constexpr TCHAR Name[] = TEXT("InviteLobbyMemberImpl");
+	};
+};
 
-		struct Params
-		{
-			// The lobby handle data.
-			TSharedPtr<FLobbyDataEOS> LobbyData;
+struct FLobbiesDeclineLobbyInvitationImpl
+{
+	static constexpr TCHAR Name[] = TEXT("DeclineLobbyInvitationImpl");
 
-			// The local user agent which will perform the action.
-			FOnlineAccountIdHandle LocalUserId;
+	struct Params
+	{
+		// The local user agent which will perform the action.
+		FOnlineAccountIdHandle LocalUserId;
 
-			// The target user for the invitation.
-			FOnlineAccountIdHandle TargetUserId;
-		};
+		// Id of the lobby for which the invitations will be declined.
+		FOnlineLobbyIdHandle LobbyId;
+	};
 
-		struct Result
-		{
-		};
+	struct Result
+	{
 	};
+};
 
-	TFuture<TDefaultErrorResult<FInviteLobbyMemberImpl>> InviteLobbyMemberImpl(FInviteLobbyMemberImpl::Params&& Params);
+struct FLobbiesKickLobbyMemberImpl
+{
+	static constexpr TCHAR Name[] = TEXT("KickLobbyMemberImpl");
 
-	struct FDeclineLobbyInvitationImpl
+	struct Params
 	{
-		static constexpr TCHAR Name[] = TEXT("DeclineLobbyInvitationImpl");
+		// The lobby handle data.
+		TSharedPtr<FLobbyDataEOS> LobbyData;
 
-		struct Params
-		{
-			// The local user agent which will perform the action.
-			FOnlineAccountIdHandle LocalUserId;
+		// The local user agent which will perform the action.
+		FOnlineAccountIdHandle LocalUserId;
 
-			// Id of the lobby for which the invitations will be declined.
-			FOnlineLobbyIdHandle LobbyId;
-		};
+		// The target user to be kicked.
+		FOnlineAccountIdHandle TargetUserId;
+	};
 
-		struct Result
-		{
-		};
+	struct Result
+	{
 	};
+};
 
-	TFuture<TDefaultErrorResult<FDeclineLobbyInvitationImpl>> DeclineLobbyInvitationImpl(FDeclineLobbyInvitationImpl::Params&& Params);
+struct FLobbiesPromoteLobbyMemberImpl
+{
+	static constexpr TCHAR Name[] = TEXT("PromoteLobbyMemberImpl");
 
-	struct FKickLobbyMemberImpl
+	struct Params
 	{
-		static constexpr TCHAR Name[] = TEXT("KickLobbyMemberImpl");
-
-		struct Params
-		{
-			// The lobby handle data.
-			TSharedPtr<FLobbyDataEOS> LobbyData;
+		// The lobby handle data.
+		TSharedPtr<FLobbyDataEOS> LobbyData;
 
-			// The local user agent which will perform the action.
-			FOnlineAccountIdHandle LocalUserId;
+		// The local user agent which will perform the action.
+		FOnlineAccountIdHandle LocalUserId;
 
-			// The target user to be kicked.
-			FOnlineAccountIdHandle TargetUserId;
-		};
+		// The target user to be promoted to owner.
+		FOnlineAccountIdHandle TargetUserId;
+	};
 
-		struct Result
-		{
-		};
+	struct Result
+	{
 	};
+};
 
-	TFuture<TDefaultErrorResult<FKickLobbyMemberImpl>> KickLobbyMemberImpl(FKickLobbyMemberImpl::Params&& Params);
+struct FLobbiesModifyLobbyDataImpl
+{
+	static constexpr TCHAR Name[] = TEXT("ModifyLobbyDataImpl");
 
-	struct FPromoteLobbyMemberImpl
+	struct Params
 	{
-		static constexpr TCHAR Name[] = TEXT("PromoteLobbyMemberImpl");
+		// The lobby handle data.
+		TSharedPtr<FLobbyDataEOS> LobbyData;
 
-		struct Params
-		{
-			// The lobby handle data.
-			TSharedPtr<FLobbyDataEOS> LobbyData;
+		// The local user agent which will perform the action.
+		FOnlineAccountIdHandle LocalUserId;
 
-			// The local user agent which will perform the action.
-			FOnlineAccountIdHandle LocalUserId;
+		// The changes to apply to the lobby data.
+		TSharedPtr<FClientLobbyDataChanges> Changes;
+	};
 
-			// The target user to be promoted to owner.
-			FOnlineAccountIdHandle TargetUserId;
-		};
+	// Todo: += operator.
+	// Mergeable op must be queued by lobby id.
 
-		struct Result
-		{
-		};
+	struct Result
+	{
 	};
+};
 
-	TFuture<TDefaultErrorResult<FPromoteLobbyMemberImpl>> PromoteLobbyMemberImpl(FPromoteLobbyMemberImpl::Params&& Params);
+struct FLobbiesModifyLobbyMemberDataImpl
+{
+	static constexpr TCHAR Name[] = TEXT("ModifyLobbyMemberDataImpl");
 
-	struct FModifyLobbyDataImpl
+	struct Params
 	{
-		static constexpr TCHAR Name[] = TEXT("ModifyLobbyDataImpl");
-
-		struct Params
-		{
-			// The lobby handle data.
-			TSharedPtr<FLobbyDataEOS> LobbyData;
+		// The lobby handle data.
+		TSharedPtr<FLobbyDataEOS> LobbyData;
 
-			// The local user agent which will perform the action.
-			FOnlineAccountIdHandle LocalUserId;
+		// The local user agent which will perform the action.
+		FOnlineAccountIdHandle LocalUserId;
 
-			// The changes to apply to the lobby data.
-			TSharedPtr<FClientLobbyDataChanges> Changes;
-		};
+		// The changes to apply to the lobby member data.
+		TSharedPtr<FClientLobbyMemberDataChanges> Changes;
+	};
 
-		// Todo: += operator.
-		// Mergeable op must be queued by lobby id.
+	// Todo: += operator.
+	// Mergeable op must be queued on union of lobby id + lobby member account id.
 
-		struct Result
-		{
-		};
+	struct Result
+	{
 	};
+};
 
-	TFuture<TDefaultErrorResult<FModifyLobbyDataImpl>> ModifyLobbyDataImpl(FModifyLobbyDataImpl::Params&& Params);
+struct FLobbiesProcessLobbyNotificationImpl
+{
+	static constexpr TCHAR Name[] = TEXT("ProcessLobbyNotificationImpl");
 
-	struct FModifyLobbyMemberDataImpl
+	struct Params
 	{
-		static constexpr TCHAR Name[] = TEXT("ModifyLobbyMemberDataImpl");
-
-		struct Params
-		{
-			// The lobby handle data.
-			TSharedPtr<FLobbyDataEOS> LobbyData;
+		// The lobby handle data.
+		TSharedPtr<FLobbyDataEOS> LobbyData;
 
-			// The local user agent which will perform the action.
-			FOnlineAccountIdHandle LocalUserId;
+		// Joining / mutated members.
+		TSet<EOS_ProductUserId> MutatedMembers;
 
-			// The changes to apply to the lobby member data.
-			TSharedPtr<FClientLobbyMemberDataChanges> Changes;
-		};
+		// Leaving members.
+		TMap<EOS_ProductUserId, ELobbyMemberLeaveReason> LeavingMembers;
+	};
 
-		// Todo: += operator.
-		// Mergeable op must be queued on union of lobby id + lobby member account id.
+	// Todo: += operator.
+	// Mergeable op must be queued by lobby id.
 
-		struct Result
-		{
-		};
+	struct Result
+	{
 	};
+};
 
-	TFuture<TDefaultErrorResult<FModifyLobbyMemberDataImpl>> ModifyLobbyMemberDataImpl(FModifyLobbyMemberDataImpl::Params&& Params);
+class FLobbiesEOSGS : public FLobbiesCommon
+{
+public:
+	FLobbiesEOSGS(FOnlineServicesEOSGS& InServices);
 
-	struct FProcessLobbyNotificationImpl
-	{
-		static constexpr TCHAR Name[] = TEXT("ProcessLobbyNotificationImpl");
+	virtual void Initialize() override;
+	virtual void PreShutdown() override;
 
-		struct Params
-		{
-			// The lobby handle data.
-			TSharedPtr<FLobbyDataEOS> LobbyData;
+	virtual TOnlineAsyncOpHandle<FCreateLobby> CreateLobby(FCreateLobby::Params&& Params) override;
+	virtual TOnlineAsyncOpHandle<FFindLobbies> FindLobbies(FFindLobbies::Params&& Params) override;
+	virtual TOnlineAsyncOpHandle<FJoinLobby> JoinLobby(FJoinLobby::Params&& Params) override;
+	virtual TOnlineAsyncOpHandle<FLeaveLobby> LeaveLobby(FLeaveLobby::Params&& Params) override;
+	virtual TOnlineAsyncOpHandle<FInviteLobbyMember> InviteLobbyMember(FInviteLobbyMember::Params&& Params) override;
+	virtual TOnlineAsyncOpHandle<FDeclineLobbyInvitation> DeclineLobbyInvitation(FDeclineLobbyInvitation::Params&& Params) override;
+	virtual TOnlineAsyncOpHandle<FKickLobbyMember> KickLobbyMember(FKickLobbyMember::Params&& Params) override;
+	virtual TOnlineAsyncOpHandle<FPromoteLobbyMember> PromoteLobbyMember(FPromoteLobbyMember::Params&& Params) override;
+	virtual TOnlineAsyncOpHandle<FModifyLobbyJoinPolicy> ModifyLobbyJoinPolicy(FModifyLobbyJoinPolicy::Params&& Params) override;
+	virtual TOnlineAsyncOpHandle<FModifyLobbyAttributes> ModifyLobbyAttributes(FModifyLobbyAttributes::Params&& Params) override;
+	virtual TOnlineAsyncOpHandle<FModifyLobbyMemberAttributes> ModifyLobbyMemberAttributes(FModifyLobbyMemberAttributes::Params&& Params) override;
+	virtual TOnlineResult<FGetJoinedLobbies> GetJoinedLobbies(FGetJoinedLobbies::Params&& Params) override;
+
+private:
+	void HandleLobbyUpdated(const EOS_Lobby_LobbyUpdateReceivedCallbackInfo* Data);
+	void HandleLobbyMemberUpdated(const EOS_Lobby_LobbyMemberUpdateReceivedCallbackInfo* Data);
+	void HandleLobbyMemberStatusReceived(const EOS_Lobby_LobbyMemberStatusReceivedCallbackInfo* Data);
+	void HandleLobbyInviteReceived(const EOS_Lobby_LobbyInviteReceivedCallbackInfo* Data);
+	void HandleLobbyInviteAccepted(const EOS_Lobby_LobbyInviteAcceptedCallbackInfo* Data);
+	void HandleJoinLobbyAccepted(const EOS_Lobby_JoinLobbyAcceptedCallbackInfo* Data);
 
-			// Joining / mutated members.
-			TSet<EOS_ProductUserId> MutatedMembers;
+protected:
+#if !UE_BUILD_SHIPPING
+	static void CheckMetadata();
+#endif
 
-			// Leaving members.
-			TMap<EOS_ProductUserId, ELobbyMemberLeaveReason> LeavingMembers;
-		};
+	void RegisterHandlers();
+	void UnregisterHandlers();
 
-		// Todo: += operator.
-		// Mergeable op must be queued by lobby id.
+	void AddActiveLobby(FOnlineAccountIdHandle LocalUserId, const TSharedRef<FLobbyDataEOS>& LobbyData);
+	void RemoveActiveLobby(FOnlineAccountIdHandle LocalUserId, const TSharedRef<FLobbyDataEOS>& LobbyData);
 
-		struct Result
-		{
-		};
-	};
+	// Todo: store list of invites per lobby.
+	void AddActiveInvite(const TSharedRef<FLobbyInviteDataEOS>& Invite);
+	void RemoveActiveInvite(const TSharedRef<FLobbyInviteDataEOS>& Invite);
+	TSharedPtr<FLobbyInviteDataEOS> GetActiveInvite(FOnlineAccountIdHandle TargetUser, FOnlineLobbyIdHandle TargetLobbyId);
 
-	TOnlineAsyncOpHandle<FProcessLobbyNotificationImpl> ProcessLobbyNotificationImplOp(FProcessLobbyNotificationImpl::Params&& Params);
+	// LobbyData will be fetched from the operation data if not set in Params.
+	TFuture<TDefaultErrorResult<FLobbiesJoinLobbyImpl>> JoinLobbyImpl(FLobbiesJoinLobbyImpl::Params&& Params);
+	TOnlineAsyncOpHandle<FLobbiesJoinLobbyMemberImpl> JoinLobbyMemberImplOp(FLobbiesJoinLobbyMemberImpl::Params&& Params);
+	TFuture<TDefaultErrorResult<FLobbiesJoinLobbyMemberImpl>> JoinLobbyMemberImpl(FLobbiesJoinLobbyMemberImpl::Params&& Params);
+	TFuture<TDefaultErrorResult<FLobbiesLeaveLobbyImpl>> LeaveLobbyImpl(FLobbiesLeaveLobbyImpl::Params&& Params);
+	TFuture<TDefaultErrorResult<FLobbiesDestroyLobbyImpl>> DestroyLobbyImpl(FLobbiesDestroyLobbyImpl::Params&& Params);
+	TFuture<TDefaultErrorResult<FLobbiesInviteLobbyMemberImpl>> InviteLobbyMemberImpl(FLobbiesInviteLobbyMemberImpl::Params&& Params);
+	TFuture<TDefaultErrorResult<FLobbiesDeclineLobbyInvitationImpl>> DeclineLobbyInvitationImpl(FLobbiesDeclineLobbyInvitationImpl::Params&& Params);
+	TFuture<TDefaultErrorResult<FLobbiesKickLobbyMemberImpl>> KickLobbyMemberImpl(FLobbiesKickLobbyMemberImpl::Params&& Params);
+	TFuture<TDefaultErrorResult<FLobbiesPromoteLobbyMemberImpl>> PromoteLobbyMemberImpl(FLobbiesPromoteLobbyMemberImpl::Params&& Params);
+	TFuture<TDefaultErrorResult<FLobbiesModifyLobbyDataImpl>> ModifyLobbyDataImpl(FLobbiesModifyLobbyDataImpl::Params&& Params);
+	TFuture<TDefaultErrorResult<FLobbiesModifyLobbyMemberDataImpl>> ModifyLobbyMemberDataImpl(FLobbiesModifyLobbyMemberDataImpl::Params&& Params);
+	TOnlineAsyncOpHandle<FLobbiesProcessLobbyNotificationImpl> ProcessLobbyNotificationImplOp(FLobbiesProcessLobbyNotificationImpl::Params&& Params);
 
 	EOSEventRegistrationPtr OnLobbyUpdatedEOSEventRegistration;
 	EOSEventRegistrationPtr OnLobbyMemberUpdatedEOSEventRegistration;
@@ -336,4 +330,103 @@ protected:
 	TMap<FOnlineAccountIdHandle, TSharedRef<FLobbySearchEOS>> ActiveSearchResults;
 };
 
+namespace Meta {
+
+BEGIN_ONLINE_STRUCT_META(FLobbiesJoinLobbyImpl::Params)
+	ONLINE_STRUCT_FIELD(FLobbiesJoinLobbyImpl::Params, LobbyData),
+	ONLINE_STRUCT_FIELD(FLobbiesJoinLobbyImpl::Params, LocalUserId),
+	ONLINE_STRUCT_FIELD(FLobbiesJoinLobbyImpl::Params, LocalName),
+	ONLINE_STRUCT_FIELD(FLobbiesJoinLobbyImpl::Params, LocalUsers)
+END_ONLINE_STRUCT_META()
+
+BEGIN_ONLINE_STRUCT_META(FLobbiesJoinLobbyImpl::Result)
+END_ONLINE_STRUCT_META()
+
+BEGIN_ONLINE_STRUCT_META(FLobbiesJoinLobbyMemberImpl::Params)
+	ONLINE_STRUCT_FIELD(FLobbiesJoinLobbyMemberImpl::Params, LobbyData),
+	ONLINE_STRUCT_FIELD(FLobbiesJoinLobbyMemberImpl::Params, LocalUserId),
+	ONLINE_STRUCT_FIELD(FLobbiesJoinLobbyMemberImpl::Params, Attributes)
+END_ONLINE_STRUCT_META()
+
+BEGIN_ONLINE_STRUCT_META(FLobbiesJoinLobbyMemberImpl::Result)
+END_ONLINE_STRUCT_META()
+
+BEGIN_ONLINE_STRUCT_META(FLobbiesLeaveLobbyImpl::Params)
+	ONLINE_STRUCT_FIELD(FLobbiesLeaveLobbyImpl::Params, LobbyData),
+	ONLINE_STRUCT_FIELD(FLobbiesLeaveLobbyImpl::Params, LocalUserId)
+END_ONLINE_STRUCT_META()
+
+BEGIN_ONLINE_STRUCT_META(FLobbiesLeaveLobbyImpl::Result)
+END_ONLINE_STRUCT_META()
+
+BEGIN_ONLINE_STRUCT_META(FLobbiesDestroyLobbyImpl::Params)
+	ONLINE_STRUCT_FIELD(FLobbiesDestroyLobbyImpl::Params, LobbyData),
+	ONLINE_STRUCT_FIELD(FLobbiesDestroyLobbyImpl::Params, LocalUserId)
+END_ONLINE_STRUCT_META()
+
+BEGIN_ONLINE_STRUCT_META(FLobbiesDestroyLobbyImpl::Result)
+END_ONLINE_STRUCT_META()
+
+BEGIN_ONLINE_STRUCT_META(FLobbiesInviteLobbyMemberImpl::Params)
+	ONLINE_STRUCT_FIELD(FLobbiesInviteLobbyMemberImpl::Params, LobbyData),
+	ONLINE_STRUCT_FIELD(FLobbiesInviteLobbyMemberImpl::Params, LocalUserId),
+	ONLINE_STRUCT_FIELD(FLobbiesInviteLobbyMemberImpl::Params, TargetUserId)
+END_ONLINE_STRUCT_META()
+
+BEGIN_ONLINE_STRUCT_META(FLobbiesInviteLobbyMemberImpl::Result)
+END_ONLINE_STRUCT_META()
+
+BEGIN_ONLINE_STRUCT_META(FLobbiesDeclineLobbyInvitationImpl::Params)
+	ONLINE_STRUCT_FIELD(FLobbiesDeclineLobbyInvitationImpl::Params, LocalUserId),
+	ONLINE_STRUCT_FIELD(FLobbiesDeclineLobbyInvitationImpl::Params, LobbyId)
+END_ONLINE_STRUCT_META()
+
+BEGIN_ONLINE_STRUCT_META(FLobbiesDeclineLobbyInvitationImpl::Result)
+END_ONLINE_STRUCT_META()
+
+BEGIN_ONLINE_STRUCT_META(FLobbiesKickLobbyMemberImpl::Params)
+	ONLINE_STRUCT_FIELD(FLobbiesKickLobbyMemberImpl::Params, LobbyData),
+	ONLINE_STRUCT_FIELD(FLobbiesKickLobbyMemberImpl::Params, LocalUserId),
+	ONLINE_STRUCT_FIELD(FLobbiesKickLobbyMemberImpl::Params, TargetUserId)
+END_ONLINE_STRUCT_META()
+
+BEGIN_ONLINE_STRUCT_META(FLobbiesKickLobbyMemberImpl::Result)
+END_ONLINE_STRUCT_META()
+
+BEGIN_ONLINE_STRUCT_META(FLobbiesPromoteLobbyMemberImpl::Params)
+	ONLINE_STRUCT_FIELD(FLobbiesPromoteLobbyMemberImpl::Params, LobbyData),
+	ONLINE_STRUCT_FIELD(FLobbiesPromoteLobbyMemberImpl::Params, LocalUserId),
+	ONLINE_STRUCT_FIELD(FLobbiesPromoteLobbyMemberImpl::Params, TargetUserId)
+END_ONLINE_STRUCT_META()
+
+BEGIN_ONLINE_STRUCT_META(FLobbiesPromoteLobbyMemberImpl::Result)
+END_ONLINE_STRUCT_META()
+
+BEGIN_ONLINE_STRUCT_META(FLobbiesModifyLobbyDataImpl::Params)
+	ONLINE_STRUCT_FIELD(FLobbiesModifyLobbyDataImpl::Params, LobbyData),
+	ONLINE_STRUCT_FIELD(FLobbiesModifyLobbyDataImpl::Params, LocalUserId)
+END_ONLINE_STRUCT_META()
+
+BEGIN_ONLINE_STRUCT_META(FLobbiesModifyLobbyDataImpl::Result)
+END_ONLINE_STRUCT_META()
+
+BEGIN_ONLINE_STRUCT_META(FLobbiesModifyLobbyMemberDataImpl::Params)
+	ONLINE_STRUCT_FIELD(FLobbiesModifyLobbyMemberDataImpl::Params, LobbyData),
+	ONLINE_STRUCT_FIELD(FLobbiesModifyLobbyMemberDataImpl::Params, LocalUserId)
+END_ONLINE_STRUCT_META()
+
+BEGIN_ONLINE_STRUCT_META(FLobbiesModifyLobbyMemberDataImpl::Result)
+END_ONLINE_STRUCT_META()
+
+BEGIN_ONLINE_STRUCT_META(FLobbiesProcessLobbyNotificationImpl::Params)
+	ONLINE_STRUCT_FIELD(FLobbiesProcessLobbyNotificationImpl::Params, LobbyData),
+	ONLINE_STRUCT_FIELD(FLobbiesProcessLobbyNotificationImpl::Params, MutatedMembers),
+	ONLINE_STRUCT_FIELD(FLobbiesProcessLobbyNotificationImpl::Params, LeavingMembers)
+END_ONLINE_STRUCT_META()
+
+BEGIN_ONLINE_STRUCT_META(FLobbiesProcessLobbyNotificationImpl::Result)
+END_ONLINE_STRUCT_META()
+
+/* Meta*/ }
+
 /* UE::Online */ }
diff --git a/Engine/Plugins/Online/Experimental/OnlineServicesOSSAdapter/Source/Private/Online/AuthOSSAdapter.cpp b/Engine/Plugins/Online/Experimental/OnlineServicesOSSAdapter/Source/Private/Online/AuthOSSAdapter.cpp
--- a/Engine/Plugins/Online/Experimental/OnlineServicesOSSAdapter/Source/Private/Online/AuthOSSAdapter.cpp
+++ b/Engine/Plugins/Online/Experimental/OnlineServicesOSSAdapter/Source/Private/Online/AuthOSSAdapter.cpp
@@ -135,7 +135,7 @@ void FAuthOSSAdapter::PostInitialize()
 				{
 					if (TSharedPtr<FAuthOSSAdapter> PinnedThis = WeakThis.Pin())
 					{
-						PinnedThis->HandleLoginStatusChangedImplOp(FHandleLoginStatusChangedImpl::Params{
+						PinnedThis->HandleLoginStatusChangedImplOp(FAuthHandleLoginStatusChangedImpl::Params{
 							FPlatformMisc::GetPlatformUserForUserIndex(LocalUserNum),
 							static_cast<FOnlineServicesOSSAdapter&>(PinnedThis->Services).GetAccountIdRegistry().FindOrAddHandle(NetId.AsShared()),
 							TranslateLoginStatus(NewStatus) });
@@ -532,6 +532,17 @@ int32 FAuthOSSAdapter::GetLocalUserNum(FOnlineAccountIdHandle AccountIdHandle) c
 	return AccountInfoOSSAdapter ? AccountInfoOSSAdapter->LocalUserNum : INDEX_NONE;
 }
 
+#if !UE_BUILD_SHIPPING
+void FAuthOSSAdapter::CheckMetadata()
+{
+	// Metadata sanity check.
+	ToLogString(FAuthHandleLoginStatusChangedImpl::Params());
+	ToLogString(FAuthHandleLoginStatusChangedImpl::Result());
+	Meta::VisitFields(FAuthHandleLoginStatusChangedImpl::Params(), [](const TCHAR* Name, auto& Field) { return false; });
+	Meta::VisitFields(FAuthHandleLoginStatusChangedImpl::Result(), [](const TCHAR* Name, auto& Field) { return false; });
+}
+#endif
+
 const FAccountInfoRegistry& FAuthOSSAdapter::GetAccountInfoRegistry() const
 {
 	return AccountInfoRegistryOSSAdapter;
@@ -557,14 +568,14 @@ IOnlineIdentityPtr FAuthOSSAdapter::GetIdentityInterface() const
 	return GetSubsystem().GetIdentityInterface();
 }
 
-TOnlineAsyncOpHandle<FAuthOSSAdapter::FHandleLoginStatusChangedImpl> FAuthOSSAdapter::HandleLoginStatusChangedImplOp(FHandleLoginStatusChangedImpl::Params&& Params)
+TOnlineAsyncOpHandle<FAuthHandleLoginStatusChangedImpl> FAuthOSSAdapter::HandleLoginStatusChangedImplOp(FAuthHandleLoginStatusChangedImpl::Params&& Params)
 {
-	TOnlineAsyncOpRef<FHandleLoginStatusChangedImpl> Op = GetOp<FHandleLoginStatusChangedImpl>(MoveTemp(Params));
+	TOnlineAsyncOpRef<FAuthHandleLoginStatusChangedImpl> Op = GetOp<FAuthHandleLoginStatusChangedImpl>(MoveTemp(Params));
 
 	// Step 1: Set up operation data.
-	Op->Then([this](TOnlineAsyncOp<FHandleLoginStatusChangedImpl>& InAsyncOp)
+	Op->Then([this](TOnlineAsyncOp<FAuthHandleLoginStatusChangedImpl>& InAsyncOp)
 	{
-		const FHandleLoginStatusChangedImpl::Params& Params = InAsyncOp.GetParams();
+		const FAuthHandleLoginStatusChangedImpl::Params& Params = InAsyncOp.GetParams();
 		TSharedPtr<FAccountInfoOSSAdapter> AccountInfoOSSAdapter = AccountInfoRegistryOSSAdapter.Find(Params.AccountId);
 		if (!AccountInfoOSSAdapter)
 		{
@@ -577,7 +588,7 @@ TOnlineAsyncOpHandle<FAuthOSSAdapter::FHandleLoginStatusChangedImpl> FAuthOSSAda
 		{
 			UE_LOG(LogOnlineServices, Log, TEXT("[FAuthOSSAdapter::HandleLoginStatusChangedImplOp][%s] Login status has not changed. Ignoring event. [%s] Current Status: %s"),
 				*GetSubsystem().GetSubsystemName().ToString(), *ToLogString(AccountInfoOSSAdapter->AccountId), *ToLogString(Params.NewLoginStatus));
-			InAsyncOp.SetResult(FHandleLoginStatusChangedImpl::Result{});
+			InAsyncOp.SetResult(FAuthHandleLoginStatusChangedImpl::Result{});
 			return;
 		}
 
@@ -585,16 +596,16 @@ TOnlineAsyncOpHandle<FAuthOSSAdapter::FHandleLoginStatusChangedImpl> FAuthOSSAda
 		InAsyncOp.Data.Set<TSharedRef<FAccountInfoOSSAdapter>>(AccountInfoKeyName, AccountInfoOSSAdapter.ToSharedRef());
 	})
 	// Step 2: Update status and notify.
-	.Then([this](TOnlineAsyncOp<FHandleLoginStatusChangedImpl>& InAsyncOp)
+	.Then([this](TOnlineAsyncOp<FAuthHandleLoginStatusChangedImpl>& InAsyncOp)
 	{
-		const FHandleLoginStatusChangedImpl::Params& Params = InAsyncOp.GetParams();
+		const FAuthHandleLoginStatusChangedImpl::Params& Params = InAsyncOp.GetParams();
 		const TSharedRef<FAccountInfoOSSAdapter>& AccountInfoOSSAdapter = GetOpDataChecked<TSharedRef<FAccountInfoOSSAdapter>>(InAsyncOp, AccountInfoKeyName);
 
 		UE_LOG(LogOnlineServices, Log, TEXT("[FAuthOSSAdapter::HandleLoginStatusChangedImplOp][%s] Login status changed for account [%s]: %s => %s."),
 			*GetSubsystem().GetSubsystemName().ToString(), *ToLogString(AccountInfoOSSAdapter->AccountId), *ToLogString(AccountInfoOSSAdapter->LoginStatus), *ToLogString(Params.NewLoginStatus));
 		AccountInfoOSSAdapter->LoginStatus = Params.NewLoginStatus;
 		OnAuthLoginStatusChangedEvent.Broadcast(FAuthLoginStatusChanged{ AccountInfoOSSAdapter, AccountInfoOSSAdapter->LoginStatus });
-		InAsyncOp.SetResult(FHandleLoginStatusChangedImpl::Result{});
+		InAsyncOp.SetResult(FAuthHandleLoginStatusChangedImpl::Result{});
 	})
 	.Enqueue(GetSerialQueue());
 
diff --git a/Engine/Plugins/Online/Experimental/OnlineServicesOSSAdapter/Source/Public/Online/AuthOSSAdapter.h b/Engine/Plugins/Online/Experimental/OnlineServicesOSSAdapter/Source/Public/Online/AuthOSSAdapter.h
--- a/Engine/Plugins/Online/Experimental/OnlineServicesOSSAdapter/Source/Public/Online/AuthOSSAdapter.h
+++ b/Engine/Plugins/Online/Experimental/OnlineServicesOSSAdapter/Source/Public/Online/AuthOSSAdapter.h
@@ -14,6 +14,22 @@ namespace UE::Online {
 
 class FOnlineServicesOSSAdapter;
 
+struct FAuthHandleLoginStatusChangedImpl
+{
+	static constexpr TCHAR Name[] = TEXT("HandleLoginStatusChangedImpl");
+
+	struct Params
+	{
+		FPlatformUserId PlatformUserId;
+		FOnlineAccountIdHandle AccountId;
+		ELoginStatus NewLoginStatus;
+	};
+
+	struct Result
+	{
+	};
+};
+
 struct FAccountInfoOSSAdapter final : public FAccountInfo
 {
 	FUniqueNetIdPtr UniqueNetId;
@@ -56,6 +72,10 @@ public:
 	int32 GetLocalUserNum(FOnlineAccountIdHandle AccountIdHandle) const;
 
 protected:
+#if !UE_BUILD_SHIPPING
+	static void CheckMetadata();
+#endif
+
 	virtual const FAccountInfoRegistry& GetAccountInfoRegistry() const override;
 
 	const FOnlineServicesOSSAdapter& GetOnlineServicesOSSAdapter() const;
@@ -63,26 +83,23 @@ protected:
 	const IOnlineSubsystem& GetSubsystem() const;
 	IOnlineIdentityPtr GetIdentityInterface() const;
 
-	struct FHandleLoginStatusChangedImpl
-	{
-		static constexpr TCHAR Name[] = TEXT("HandleLoginStatusChangedImpl");
-
-		struct Params
-		{
-			FPlatformUserId PlatformUserId;
-			FOnlineAccountIdHandle AccountId;
-			ELoginStatus NewLoginStatus;
-		};
-
-		struct Result
-		{
-		};
-	};
-
-	TOnlineAsyncOpHandle<FHandleLoginStatusChangedImpl> HandleLoginStatusChangedImplOp(FHandleLoginStatusChangedImpl::Params&& Params);
+	TOnlineAsyncOpHandle<FAuthHandleLoginStatusChangedImpl> HandleLoginStatusChangedImplOp(FAuthHandleLoginStatusChangedImpl::Params&& Params);
 
 	FAccountInfoRegistryOSSAdapter AccountInfoRegistryOSSAdapter;
 	FDelegateHandle OnLoginStatusChangedHandle[MAX_LOCAL_PLAYERS];
 };
 
+namespace Meta {
+
+BEGIN_ONLINE_STRUCT_META(FAuthHandleLoginStatusChangedImpl::Params)
+	ONLINE_STRUCT_FIELD(FAuthHandleLoginStatusChangedImpl::Params, PlatformUserId),
+	ONLINE_STRUCT_FIELD(FAuthHandleLoginStatusChangedImpl::Params, AccountId),
+	ONLINE_STRUCT_FIELD(FAuthHandleLoginStatusChangedImpl::Params, NewLoginStatus)
+END_ONLINE_STRUCT_META()
+
+BEGIN_ONLINE_STRUCT_META(FAuthHandleLoginStatusChangedImpl::Result)
+END_ONLINE_STRUCT_META()
+
+/* Meta*/ }
+
 /* UE::Online */ }
