Add null terminator space for strings

diff --git a/Engine/Plugins/Online/Experimental/OnlineServicesEOSGS/Source/Private/Online/AuthEOSGS.cpp b/Engine/Plugins/Online/Experimental/OnlineServicesEOSGS/Source/Private/Online/AuthEOSGS.cpp
--- a/Engine/Plugins/Online/Experimental/OnlineServicesEOSGS/Source/Private/Online/AuthEOSGS.cpp
+++ b/Engine/Plugins/Online/Experimental/OnlineServicesEOSGS/Source/Private/Online/AuthEOSGS.cpp
@@ -85,7 +85,7 @@ private:
 
 	EOS_Connect_Credentials CredentialsData;
 	EOS_Connect_UserLoginInfo UserLoginInfoData;
-	char DisplayNameUtf8[EOS_CONNECT_USERLOGININFO_DISPLAYNAME_MAX_LENGTH] = {};
+	char DisplayNameUtf8[EOS_CONNECT_USERLOGININFO_DISPLAYNAME_MAX_LENGTH + 1] = {};
 	TArray<char> ExternalAuthTokenUtf8;
 };
 
@@ -138,7 +138,7 @@ FEOSConnectLoginOptions& FEOSConnectLoginOptions::operator=(FEOSConnectLoginOpti
 
 	if (Other.UserLoginInfo)
 	{
-		memcpy(DisplayNameUtf8, Other.DisplayNameUtf8, EOS_CONNECT_USERLOGININFO_DISPLAYNAME_MAX_LENGTH);
+		memcpy(DisplayNameUtf8, Other.DisplayNameUtf8, sizeof(DisplayNameUtf8));
 		UserLoginInfoData = Other.UserLoginInfoData;
 		UserLoginInfoData.DisplayName = DisplayNameUtf8;
 		UserLoginInfo = &UserLoginInfoData;
@@ -177,7 +177,7 @@ TDefaultErrorResultInternal<FEOSConnectLoginOptions> FEOSConnectLoginOptions::Cr
 	if (EnumHasAnyFlags(TranslatorTraits->Flags, EEOSConnectTranslationFlags::DisplayName))
 	{
 		// Todo: Lookup platform OSS Auth interface to resolve display name using PlatformUserId.
-		FCStringAnsi::Strncpy(EOSConnectLoginOptions.DisplayNameUtf8, "PlatformUser", EOS_CONNECT_USERLOGININFO_DISPLAYNAME_MAX_LENGTH);
+		FCStringAnsi::Strncpy(EOSConnectLoginOptions.DisplayNameUtf8, "PlatformUser", sizeof(EOSConnectLoginOptions.DisplayNameUtf8));
 		EOSConnectLoginOptions.UserLoginInfo = &EOSConnectLoginOptions.UserLoginInfoData;
 	}
 
@@ -753,8 +753,8 @@ TOnlineAsyncOpHandle<FAuthLogin> FAuthEOSGS::Login(FAuthLogin::Params&& Params)
 		if (!IsOnlineStatus(AccountInfoEOS->LoginStatus))
 		{
 			// Set DisplayName.
-			int32_t ProductUserIdBufferLength = EOS_PRODUCTUSERID_MAX_LENGTH;
-			char ProductUserIdBuffer[EOS_PRODUCTUSERID_MAX_LENGTH] = {};
+			int32_t ProductUserIdBufferLength = EOS_PRODUCTUSERID_MAX_LENGTH + 1;
+			char ProductUserIdBuffer[EOS_PRODUCTUSERID_MAX_LENGTH + 1] = {};
 			EOS_EResult ProductUserIdResult = EOS_ProductUserId_ToString(AccountInfoEOS->ProductUserId, ProductUserIdBuffer, &ProductUserIdBufferLength);
 			if (ProductUserIdResult == EOS_EResult::EOS_Success)
 			{
