Wrap-up of EOSGS Sessions Interface implementation:

OI-4407: EOS Session Interface: LAN Support
- Moved common LAN logic from SessionsNull to SessionsLAN.
- Added logic to SessionsEOSGS to use this common LAN logic.
- Added NboSerializerEOSGSSvc to support LAN capabilities.

OI-4775: EOS Sessions Implementation: Remaining API coverage
- Added AddSessionMembers and RemoveSessionMembers methods to the Sessions Interface.
- Refactored RegisterPlayers and UnregisterPlayers to work alongside them.
- Implemented AddSessionMembers and RemoveSessionMembers in EOSGS Sessions.
- Added SessionsCommon parameter and state checks to sessions methods.

General refactoring:
- Added += operator to FSessionSettings to reduce code duplication in UpdateSession methods.
- As part of RegisterPlayers refactor, removed usage of NumPrivateConnections until it has been discussed.

diff --git a/Engine/Plugins/Online/Experimental/OnlineServices/Source/OnlineServicesCommon/Private/Online/SessionsCommon.cpp b/Engine/Plugins/Online/Experimental/OnlineServices/Source/OnlineServicesCommon/Private/Online/SessionsCommon.cpp
--- a/Engine/Plugins/Online/Experimental/OnlineServices/Source/OnlineServicesCommon/Private/Online/SessionsCommon.cpp
+++ b/Engine/Plugins/Online/Experimental/OnlineServices/Source/OnlineServicesCommon/Private/Online/SessionsCommon.cpp
@@ -115,6 +115,66 @@ namespace UE::Online {
 		return Operation->GetHandle();
 	}
 
+	TOnlineAsyncOpHandle<FAddSessionMembers> FSessionsCommon::AddSessionMembers(FAddSessionMembers::Params&& Params)
+	{
+		TOnlineAsyncOpRef<FAddSessionMembers> Op = GetOp<FAddSessionMembers>(MoveTemp(Params));
+
+		Op->Then([this](TOnlineAsyncOp<FAddSessionMembers>& Op) mutable
+		{
+			const FAddSessionMembers::Params& OpParams = Op.GetParams();
+
+			FOnlineError StateCheck = CheckAddSessionMembersState(OpParams);
+			if (StateCheck != Errors::Success())
+			{
+				Op.SetError(MoveTemp(StateCheck));
+				return;
+			}
+
+			TOnlineResult<FAddSessionMembers> Result = AddSessionMembersImpl(OpParams);
+			if (Result.IsOk())
+			{
+				Op.SetResult(MoveTemp(Result.GetOkValue()));
+			}
+			else
+			{
+				Op.SetError(MoveTemp(Result.GetErrorValue()));
+			}
+		})
+		.Enqueue(GetSerialQueue());
+
+		return Op->GetHandle();
+	}
+
+	TOnlineAsyncOpHandle<FRemoveSessionMembers> FSessionsCommon::RemoveSessionMembers(FRemoveSessionMembers::Params&& Params)
+	{
+		TOnlineAsyncOpRef<FRemoveSessionMembers> Op = GetOp<FRemoveSessionMembers>(MoveTemp(Params));
+
+		Op->Then([this](TOnlineAsyncOp<FRemoveSessionMembers>& Op) mutable
+		{
+			const FRemoveSessionMembers::Params& OpParams = Op.GetParams();
+
+			FOnlineError StateCheck = CheckRemoveSessionMembersState(OpParams);
+			if (StateCheck != Errors::Success())
+			{
+				Op.SetError(MoveTemp(StateCheck));
+				return;
+			}
+
+			TOnlineResult<FRemoveSessionMembers> Result = RemoveSessionMembersImpl(OpParams);
+			if (Result.IsOk())
+			{
+				Op.SetResult(MoveTemp(Result.GetOkValue()));
+			}
+			else
+			{
+				Op.SetError(MoveTemp(Result.GetErrorValue()));
+			}
+		})
+		.Enqueue(GetSerialQueue());
+
+		return Op->GetHandle();
+	}
+
 	TOnlineAsyncOpHandle<FSendSessionInvite> FSessionsCommon::SendSessionInvite(FSendSessionInvite::Params&& Params)
 	{
 		TOnlineAsyncOpRef<FSendSessionInvite> Operation = GetOp<FSendSessionInvite>(MoveTemp(Params));
@@ -138,45 +198,21 @@ namespace UE::Online {
 		{
 			const FRegisterPlayers::Params& OpParams = Op.GetParams();
 
-			if (TSharedRef<FSession>* FoundSession = SessionsByName.Find(OpParams.SessionName))
+			FOnlineError StateCheck = CheckRegisterPlayersState(OpParams);
+			if (StateCheck != Errors::Success())
 			{
-				FSessionSettings& SessionSettings = FoundSession->Get().SessionSettings;
-
-				for (const FOnlineAccountIdHandle& User : OpParams.TargetUsers)
-				{
-					FRegisteredUser& RegisteredUser = SessionSettings.RegisteredUsers.FindOrAdd(User);
-
-					// RegisterPlayers will not revoke a slot reserved for a registered user
-					if (!RegisteredUser.bHasReservedSlot && OpParams.bReserveSlot)
-					{
-						bool bSlotReserved = false;
-						if (SessionSettings.NumOpenPublicConnections > 0)
-						{
-							SessionSettings.NumOpenPublicConnections--;
-							bSlotReserved = true;
-						}
-						else if (SessionSettings.NumOpenPrivateConnections > 0)
-						{
-							SessionSettings.NumOpenPrivateConnections--;
-							bSlotReserved = true;
-						}
-
-						RegisteredUser.bHasReservedSlot = bSlotReserved;
-					}
-
-					if (const FSessionMember* SessionMember = SessionSettings.SessionMembers.Find(User))
-					{
-						RegisteredUser.bIsInSession = true;
-					}
-				}
+				Op.SetError(MoveTemp(StateCheck));
+				return;
+			}
 
-				Op.SetResult(FRegisterPlayers::Result{});
+			TOnlineResult<FRegisterPlayers> Result = RegisterPlayersImpl(OpParams);
+			if (Result.IsOk())
+			{
+				Op.SetResult(MoveTemp(Result.GetOkValue()));
 			}
 			else
 			{
-				UE_LOG(LogTemp, Warning, TEXT("[FSessionsCommon::RegisterPlayers] Unable to register players. Session with name [%s] not found"), *OpParams.SessionName.ToString());
-
-				Op.SetError(Errors::NotFound());
+				Op.SetError(MoveTemp(Result.GetErrorValue()));
 			}
 		})
 		.Enqueue(GetSerialQueue());
@@ -193,52 +229,21 @@ namespace UE::Online {
 		{
 			const FUnregisterPlayers::Params& OpParams = Op.GetParams();
 
-			if (TSharedRef<FSession>* FoundSession = SessionsByName.Find(OpParams.SessionName))
+			FOnlineError StateCheck = CheckUnregisterPlayersState(OpParams);
+			if (StateCheck != Errors::Success())
 			{
-				FSessionSettings& SessionSettings = FoundSession->Get().SessionSettings;
-
-				for (const FOnlineAccountIdHandle& User : OpParams.TargetUsers)
-				{
-					if (const FRegisteredUser* RegisteredUser = SessionSettings.RegisteredUsers.Find(User))
-					{
-						SessionSettings.RegisteredUsers.Remove(User);
-
-						// TODO: This logic will move to RemoveSessionMember
-						bool bRemoveUserFromSession = RegisteredUser->bIsInSession && OpParams.bRemoveUnregisteredPlayers;
-
-						if (bRemoveUserFromSession)
-						{
-							SessionSettings.SessionMembers.Remove(User);
-						}
-
-						if (RegisteredUser->bHasReservedSlot || bRemoveUserFromSession)
-						{
-							if (SessionSettings.NumOpenPublicConnections < SessionSettings.NumMaxPublicConnections)
-							{
-								SessionSettings.NumOpenPublicConnections++;
-							}
-							else if (SessionSettings.NumOpenPrivateConnections < SessionSettings.NumMaxPrivateConnections)
-							{
-								SessionSettings.NumOpenPrivateConnections++;
-							}
-						}
-					}
-					else
-					{
-						UE_LOG(LogTemp, Warning, TEXT("[FSessionsCommon::UnregisterPlayers] Unable to unregister player with id [%s]. Player not registered"), *ToLogString(User));
-
-						Op.SetError(Errors::NotFound());
-						return;
-					}
-				}
+				Op.SetError(MoveTemp(StateCheck));
+				return;
+			}
 
-				Op.SetResult(FUnregisterPlayers::Result{});
+			TOnlineResult<FUnregisterPlayers> Result = UnregisterPlayersImpl(OpParams);
+			if (Result.IsOk())
+			{
+				Op.SetResult(MoveTemp(Result.GetOkValue()));
 			}
 			else
 			{
-				UE_LOG(LogTemp, Warning, TEXT("[FSessionsCommon::UnregisterPlayers] Unable to unregister players. Session with name [%s] not found"), *OpParams.SessionName.ToString());
-
-				Op.SetError(Errors::NotFound());
+				Op.SetError(MoveTemp(Result.GetErrorValue()));
 			}
 		})
 		.Enqueue(GetSerialQueue());
@@ -271,6 +276,201 @@ namespace UE::Online {
 		return SessionEvents.OnUISessionJoinRequested;
 	}
 
+	TOnlineResult<FAddSessionMembers> FSessionsCommon::AddSessionMembersImpl(const FAddSessionMembers::Params& Params)
+	{
+		TSharedRef<FSession> FoundSession = SessionsByName.FindChecked(Params.SessionName);
+
+		FSessionSettings& SessionSettings = FoundSession->SessionSettings;
+
+		for (const TPair<FOnlineAccountIdHandle, FSessionMember>& Entry : Params.NewSessionMembers)
+		{
+			const FOnlineAccountIdHandle& SessionMemberId = Entry.Key;
+			const FSessionMember& SessionMember = Entry.Value;
+
+			// If the user was already registered in the session, we'll update it and reserve a slot for it
+			if (FRegisteredPlayer* RegisteredPlayer = SessionSettings.RegisteredPlayers.Find(SessionMemberId))
+			{
+				RegisteredPlayer->bIsInSession = true;
+
+				if (!RegisteredPlayer->bHasReservedSlot)
+				{
+					bool bSlotReserved = false;
+
+					if (SessionSettings.NumOpenPublicConnections > 0)
+					{
+						SessionSettings.NumOpenPublicConnections--;
+						bSlotReserved = true;
+					}
+					else
+					{
+						return TOnlineResult<FAddSessionMembers>(Errors::InvalidState());
+					}
+
+					// TODO: Discuss usage for Private Connections
+
+					RegisteredPlayer->bHasReservedSlot = bSlotReserved;
+				}
+			}
+			else if (Params.bRegisterPlayers) // If it wasn't and the parameters dictate it, we'll register them with a reserved slot
+			{
+				FRegisterPlayers::Params RegisterPlayerParams = { Params.SessionName, { SessionMemberId }, true };
+
+				TOnlineResult<FRegisterPlayers> Result = RegisterPlayersImpl(RegisterPlayerParams);
+				if (Result.IsError())
+				{
+					return TOnlineResult<FAddSessionMembers>(MoveTemp(Result.GetErrorValue()));
+				}
+			}
+			else // Even if the new session member is not a registered player, we'll still need to decrease the number of open slots
+			{
+				if (SessionSettings.NumOpenPublicConnections > 0)
+				{
+					SessionSettings.NumOpenPublicConnections--;
+				}
+				else
+				{
+					return TOnlineResult<FAddSessionMembers>(Errors::InvalidState());
+				}
+
+				// TODO: Discuss usage for Private Connections
+			}
+
+			SessionSettings.SessionMembers.Emplace(SessionMemberId, SessionMember);
+		}
+
+		return TOnlineResult<FAddSessionMembers>();
+	}
+
+	TOnlineResult<FRemoveSessionMembers> FSessionsCommon::RemoveSessionMembersImpl(const FRemoveSessionMembers::Params& Params)
+	{
+		TSharedRef<FSession> FoundSession = SessionsByName.FindChecked(Params.SessionName);
+
+		FSessionSettings& SessionSettings = FoundSession->SessionSettings;
+
+		for (const FOnlineAccountIdHandle& SessionMemberId : Params.SessionMemberIds)
+		{
+			if (Params.bUnregisterPlayers) // If the parameters dictate it, well unregister the player as well
+			{
+				FUnregisterPlayers::Params UnregisterPlayersParams = { Params.SessionName, { SessionMemberId }, false };
+
+				TOnlineResult<FUnregisterPlayers> Result = UnregisterPlayersImpl(UnregisterPlayersParams);
+				if (Result.IsError())
+				{
+					return TOnlineResult<FRemoveSessionMembers>(MoveTemp(Result.GetErrorValue()));
+				}
+			}
+			else if (FRegisteredPlayer* RegisteredPlayer = SessionSettings.RegisteredPlayers.Find(SessionMemberId)) // If the removed session member is still registered, we'll update it and keep their reserved slot
+			{
+				RegisteredPlayer->bIsInSession = false;
+			}
+			else  // If the removed session member is not a registered player with a reserved slot, we'll increase the number of open slots
+			{
+				if (SessionSettings.NumOpenPublicConnections < SessionSettings.NumMaxPublicConnections)
+				{
+					SessionSettings.NumOpenPublicConnections++;
+				}
+				else
+				{
+					return TOnlineResult<FRemoveSessionMembers>(Errors::InvalidState());
+				}
+
+				// TODO: Discuss usage for Private Connections
+			}
+
+			SessionSettings.SessionMembers.Remove(SessionMemberId);
+		}
+
+		return TOnlineResult<FRemoveSessionMembers>();
+	}
+
+	TOnlineResult<FRegisterPlayers> FSessionsCommon::RegisterPlayersImpl(const FRegisterPlayers::Params& Params)
+	{
+		TSharedRef<FSession> FoundSession = SessionsByName.FindChecked(Params.SessionName);
+
+		FSessionSettings& SessionSettings = FoundSession->SessionSettings;
+
+		for (const FOnlineAccountIdHandle& User : Params.TargetUsers)
+		{
+			FRegisteredPlayer& RegisteredPlayer = SessionSettings.RegisteredPlayers.FindOrAdd(User);
+
+			if (const FSessionMember* SessionMember = SessionSettings.SessionMembers.Find(User))
+			{
+				RegisteredPlayer.bIsInSession = true;
+			}
+
+			// RegisterPlayers will not revoke a slot reserved for a registered player
+			if (!RegisteredPlayer.bIsInSession && !RegisteredPlayer.bHasReservedSlot && Params.bReserveSlot)
+			{
+				bool bSlotReserved = false;
+				
+				if (SessionSettings.NumOpenPublicConnections > 0)
+				{
+					SessionSettings.NumOpenPublicConnections--;
+					bSlotReserved = true;
+				}
+				else
+				{
+					return TOnlineResult<FRegisterPlayers>(Errors::InvalidState());
+				}
+
+				// TODO: Discuss usage for Private Connections
+
+				RegisteredPlayer.bHasReservedSlot = bSlotReserved;
+			}			
+		}
+
+		return TOnlineResult<FRegisterPlayers>();
+	}
+
+	TOnlineResult<FUnregisterPlayers> FSessionsCommon::UnregisterPlayersImpl(const FUnregisterPlayers::Params& Params)
+	{
+		TSharedRef<FSession> FoundSession = SessionsByName.FindChecked(Params.SessionName);
+
+		FSessionSettings& SessionSettings = FoundSession->SessionSettings;
+
+		for (const FOnlineAccountIdHandle& User : Params.TargetUsers)
+		{
+			if (const FRegisteredPlayer* RegisteredPlayer = SessionSettings.RegisteredPlayers.Find(User))
+			{
+				// If the registered player had a reserved slot, we'll remove it
+				if (RegisteredPlayer->bHasReservedSlot)
+				{
+					if (SessionSettings.NumOpenPublicConnections < SessionSettings.NumMaxPublicConnections)
+					{
+						SessionSettings.NumOpenPublicConnections++;
+					}
+					else
+					{
+						return TOnlineResult<FUnregisterPlayers>(Errors::InvalidState());
+					}
+
+					// TODO: Discuss usage for Private Connections
+				}
+
+				if (RegisteredPlayer->bIsInSession && Params.bRemoveUnregisteredPlayers)
+				{
+					FRemoveSessionMembers::Params RemoveSessionMembersParams = { User, Params.SessionName, Params.TargetUsers, false };
+
+					TOnlineResult<FRemoveSessionMembers> Result = RemoveSessionMembersImpl(RemoveSessionMembersParams);
+					if (Result.IsError())
+					{
+						return TOnlineResult<FUnregisterPlayers>(MoveTemp(Result.GetErrorValue()));
+					}
+				}
+
+				SessionSettings.RegisteredPlayers.Remove(User);
+			}
+			else
+			{
+				UE_LOG(LogTemp, Warning, TEXT("[FSessionsCommon::UnregisterPlayers] Unable to unregister player with id [%s]. Player not registered"), *ToLogString(User));
+
+				return TOnlineResult<FUnregisterPlayers>(Errors::NotFound());
+			}
+		}
+
+		return TOnlineResult<FUnregisterPlayers>();
+	}
+
 	FOnlineError FSessionsCommon::CheckCreateSessionParams(const FCreateSession::Params& Params)
 	{
 		if (Params.SessionSettings.NumMaxPrivateConnections == 0 && Params.SessionSettings.NumMaxPublicConnections == 0)
@@ -292,20 +492,15 @@ namespace UE::Online {
 			return Errors::InvalidState();
 		}
 
-		if (!Params.SessionSettings.IsDedicatedServerSession)
+		if (Params.SessionSettings.bPresenceEnabled)
 		{
-			IAuthPtr Auth = Services.GetAuthInterface();
-
-			FAuthGetLocalOnlineUserByOnlineAccountId::Params GetLocalOnlineUserByOnlineAccountIdParams{ Params.LocalUserId };
-			TOnlineResult<FAuthGetLocalOnlineUserByOnlineAccountId> Result = Auth->GetLocalOnlineUserByOnlineAccountId(MoveTemp(GetLocalOnlineUserByOnlineAccountIdParams));
-
-			if (Result.IsOk())
+			for (const TPair<FName, TSharedRef<FSession>>& Entry : SessionsByName)
 			{
-				if (Result.GetOkValue().AccountInfo->LoginStatus == ELoginStatus::NotLoggedIn)
+				if (Entry.Value->SessionSettings.bPresenceEnabled)
 				{
-					UE_LOG(LogTemp, Warning, TEXT("[FSessionsCommon::CheckCreateSessionState] Could not create session with user [%s] not logged in"), *ToLogString(Params.LocalUserId));
+					UE_LOG(LogTemp, Warning, TEXT("[FSessionsCommon::CheckCreateSessionState] Could not create session with bPresenceEnabled set to true when another already exists [%s]."), *Entry.Key.ToString());
 
-					return Errors::InvalidUser();
+					return Errors::InvalidState();
 				}
 			}
 		}
@@ -319,17 +514,11 @@ namespace UE::Online {
 
 		for (const FOnlineAccountIdHandle& LocalUserId : LocalUserIds)
 		{
-			FAuthGetLocalOnlineUserByOnlineAccountId::Params GetLocalOnlineUserByOnlineAccountIdParams{ LocalUserId };
-			TOnlineResult<FAuthGetLocalOnlineUserByOnlineAccountId> Result = Auth->GetLocalOnlineUserByOnlineAccountId(MoveTemp(GetLocalOnlineUserByOnlineAccountIdParams));
-
-			if (Result.IsOk())
+			if (!Auth->IsLoggedIn(LocalUserId))
 			{
-				if (Result.GetOkValue().AccountInfo->LoginStatus == ELoginStatus::NotLoggedIn)
-				{
-					UE_LOG(LogTemp, Warning, TEXT("[FSessionsCommon::CheckCreateSessionState] Could not create session with user [%s] not logged in"), *ToLogString(LocalUserId));
+				UE_LOG(LogTemp, Warning, TEXT("[FSessionsCommon::CheckCreateSessionState] Could not create session with user [%s] not logged in"), *ToLogString(LocalUserId));
 
-					return Errors::InvalidUser();
-				}
+				return Errors::InvalidUser();
 			}
 		}
 
@@ -340,21 +529,13 @@ namespace UE::Online {
 	{
 		if (TSharedRef<FSession>* FoundSession = SessionsByName.Find(Params.SessionName))
 		{
-			if (!FoundSession->Get().SessionSettings.IsDedicatedServerSession)
+			if (!FoundSession->Get().SessionSettings.bIsDedicatedServerSession)
 			{
-				IAuthPtr Auth = Services.GetAuthInterface();
-
-				FAuthGetLocalOnlineUserByOnlineAccountId::Params GetLocalOnlineUserByOnlineAccountIdParams{ Params.LocalUserId };
-				TOnlineResult<FAuthGetLocalOnlineUserByOnlineAccountId> Result = Auth->GetLocalOnlineUserByOnlineAccountId(MoveTemp(GetLocalOnlineUserByOnlineAccountIdParams));
-
-				if (Result.IsOk())
+				if (!Services.GetAuthInterface()->IsLoggedIn(Params.LocalUserId))
 				{
-					if (Result.GetOkValue().AccountInfo->LoginStatus == ELoginStatus::NotLoggedIn)
-					{
-						UE_LOG(LogTemp, Warning, TEXT("[FSessionsNull::CheckUpdateSessionState] Could not update session with user [%s] not logged in"), *ToLogString(Params.LocalUserId));
+					UE_LOG(LogTemp, Warning, TEXT("[FSessionsNull::CheckUpdateSessionState] Could not update session with user [%s] not logged in"), *ToLogString(Params.LocalUserId));
 
-						return Errors::InvalidUser();
-					}
+					return Errors::InvalidUser();
 				}
 			}
 		}
@@ -383,19 +564,11 @@ namespace UE::Online {
 	FOnlineError FSessionsCommon::CheckFindSessionsState(const FFindSessions::Params& Params)
 	{
 		// User login check
-		IAuthPtr Auth = Services.GetAuthInterface();
-
-		FAuthGetLocalOnlineUserByOnlineAccountId::Params GetLocalOnlineUserByOnlineAccountIdParams{ Params.LocalUserId };
-		TOnlineResult<FAuthGetLocalOnlineUserByOnlineAccountId> Result = Auth->GetLocalOnlineUserByOnlineAccountId(MoveTemp(GetLocalOnlineUserByOnlineAccountIdParams));
-
-		if (Result.IsOk())
+		if (!Services.GetAuthInterface()->IsLoggedIn(Params.LocalUserId))
 		{
-			if (Result.GetOkValue().AccountInfo->LoginStatus == ELoginStatus::NotLoggedIn)
-			{
-				UE_LOG(LogTemp, Warning, TEXT("[FSessionsCommon::CheckFindSessionsState] Could not find sessions with user [%s] not logged in"), *ToLogString(Params.LocalUserId));
+			UE_LOG(LogTemp, Warning, TEXT("[FSessionsCommon::CheckFindSessionsState] Could not find sessions with user [%s] not logged in"), *ToLogString(Params.LocalUserId));
 
-				return Errors::InvalidUser();
-			}
+			return Errors::InvalidUser();
 		}
 
 		// Ongoing search check
@@ -413,6 +586,13 @@ namespace UE::Online {
 	{
 		const FSessionSettings& SessionSettings = Params.Session->SessionSettings;
 
+		if (Params.Session->CurrentState != ESessionState::Valid)
+		{
+			UE_LOG(LogTemp, Warning, TEXT("[FSessionsEOSGS::FJoinSession] Could not join session with invalid session state [%s]"), LexToString(Params.Session->CurrentState));
+
+			return Errors::InvalidParams();
+		}
+
 		TArray<FOnlineAccountIdHandle> LocalUserIds;
 		LocalUserIds.Reserve(Params.LocalUsers.Num());
 		Params.LocalUsers.GenerateKeyArray(LocalUserIds);
@@ -438,7 +618,7 @@ namespace UE::Online {
 		{
 			for (const FOnlineAccountIdHandle& LocalUserId : LocalUserIds)
 			{
-				if (!SessionSettings.RegisteredUsers.Contains(LocalUserId))
+				if (!SessionSettings.RegisteredPlayers.Contains(LocalUserId))
 				{
 					UE_LOG(LogTemp, Warning, TEXT("[FSessionsCommon::CheckJoinSessionParams] Could not join session. bAllowUnregisteredPlayers is set to false and user [%s] is not registered"), *ToLogString(LocalUserId));
 
@@ -449,14 +629,14 @@ namespace UE::Online {
 
 		if (SessionSettings.NumMaxPublicConnections > 0 && SessionSettings.NumOpenPublicConnections < (uint32)LocalUserIds.Num())
 		{
-			UE_LOG(LogTemp, Warning, TEXT("[FSessionsCommon::CheckJoinSessionParams] Could not join session. Not enough NumOpenPublicConnections [%s] for the amount of joining players [%d]"), SessionSettings.NumOpenPublicConnections, LocalUserIds.Num());
+			UE_LOG(LogTemp, Warning, TEXT("[FSessionsCommon::CheckJoinSessionParams] Could not join session. Not enough NumOpenPublicConnections [%u] for the amount of joining players [%d]"), SessionSettings.NumOpenPublicConnections, LocalUserIds.Num());
 
 			return Errors::AccessDenied();
 		}
 
 		if (SessionSettings.NumMaxPrivateConnections > 0 && SessionSettings.NumOpenPrivateConnections < (uint32)LocalUserIds.Num())
 		{
-			UE_LOG(LogTemp, Warning, TEXT("[FSessionsCommon::CheckJoinSessionParams] Could not join session. Not enough NumOpenPrivateConnections [%s] for the amount of joining players [%d]"), SessionSettings.NumOpenPrivateConnections, LocalUserIds.Num());
+			UE_LOG(LogTemp, Warning, TEXT("[FSessionsCommon::CheckJoinSessionParams] Could not join session. Not enough NumOpenPrivateConnections [%u] for the amount of joining players [%d]"), SessionSettings.NumOpenPrivateConnections, LocalUserIds.Num());
 
 			return Errors::AccessDenied();
 		}
@@ -483,73 +663,147 @@ namespace UE::Online {
 
 		for (const FOnlineAccountIdHandle& LocalUserId : LocalUserIds)
 		{
-			FAuthGetLocalOnlineUserByOnlineAccountId::Params GetLocalOnlineUserByOnlineAccountIdParams{ LocalUserId };
-			TOnlineResult<FAuthGetLocalOnlineUserByOnlineAccountId> Result = Auth->GetLocalOnlineUserByOnlineAccountId(MoveTemp(GetLocalOnlineUserByOnlineAccountIdParams));
-
-			if (Result.IsOk())
+			if(!Auth->IsLoggedIn(LocalUserId))
 			{
-				if (Result.GetOkValue().AccountInfo->LoginStatus == ELoginStatus::NotLoggedIn)
-				{
-					UE_LOG(LogTemp, Warning, TEXT("[FSessionsCommon::CheckJoinSessionState] Could not join session with user [%s] not logged in"), *ToLogString(LocalUserId));
+				UE_LOG(LogTemp, Warning, TEXT("[FSessionsCommon::CheckJoinSessionState] Could not join session with user [%s] not logged in"), *ToLogString(LocalUserId));
 
-					return Errors::InvalidUser();
-				}
+				return Errors::InvalidUser();
 			}
 		}
 
 		return Errors::Success();
 	}
 
-	FOnlineError FSessionsCommon::CheckLeaveSessionState(const FLeaveSession::Params& Params)
+	FOnlineError FSessionsCommon::CheckAddSessionMembersState(const FAddSessionMembers::Params& Params)
 	{
-		// User login check for main caller, session check
-		if (TSharedRef<FSession>* FoundSession = SessionsByName.Find(Params.SessionName))
+		if (TOptional<FOnlineError> Result = CheckSessionExistsByName(Params.SessionName))
 		{
-			if (!FoundSession->Get().SessionSettings.IsDedicatedServerSession)
-			{
-				IAuthPtr Auth = Services.GetAuthInterface();
+			UE_LOG(LogTemp, Warning, TEXT("[FSessionsCommon::CheckAddSessionMembersState] Could not add session member to session with name [%s]. Session not found"), *Params.SessionName.ToString());
 
-				FAuthGetLocalOnlineUserByOnlineAccountId::Params GetLocalOnlineUserByOnlineAccountIdParams{ Params.LocalUserId };
-				TOnlineResult<FAuthGetLocalOnlineUserByOnlineAccountId> Result = Auth->GetLocalOnlineUserByOnlineAccountId(MoveTemp(GetLocalOnlineUserByOnlineAccountIdParams));
+			return Result.GetValue();
+		}
 
-				if (Result.IsOk())
-				{
-					if (Result.GetOkValue().AccountInfo->LoginStatus == ELoginStatus::NotLoggedIn)
-					{
-						UE_LOG(LogTemp, Warning, TEXT("[FSessionsNull::CheckLeaveSessionState] Could not leave session with user [%s] not logged in"), *ToLogString(Params.LocalUserId));
+		// TODO: Check if there are enough slots available if Params::bReserveSlot is true
 
-						return Errors::InvalidUser();
-					}
-				}
-			}
+		return Errors::Success();
+	}
+
+	FOnlineError FSessionsCommon::CheckRemoveSessionMembersState(const FRemoveSessionMembers::Params& Params)
+	{
+		if (TOptional<FOnlineError> Result = CheckSessionExistsByName(Params.SessionName))
+		{
+			UE_LOG(LogTemp, Warning, TEXT("[FSessionsCommon::CheckRemoveSessionMembersState] Could not remove session member from session with name [%s]. Session not found"), *Params.SessionName.ToString());
+
+			return Result.GetValue();
 		}
-		else
+
+		return Errors::Success();
+	}
+
+	FOnlineError FSessionsCommon::CheckLeaveSessionState(const FLeaveSession::Params& Params)
+	{
+		// User login check for main caller, session check
+		if (!SessionsByName.Contains(Params.SessionName))
 		{
 			UE_LOG(LogTemp, Warning, TEXT("[FSessionsCommon::CheckLeaveSessionState] Could not leave session with name [%s]. Session not found"), *Params.SessionName.ToString());
 
 			return Errors::NotFound();
 		}
 
-		// User login check for all local users
 		IAuthPtr Auth = Services.GetAuthInterface();
 
+		// User login check for all local users
 		for (const FOnlineAccountIdHandle& LocalUserId : Params.LocalUsers)
 		{
-			FAuthGetLocalOnlineUserByOnlineAccountId::Params GetLocalOnlineUserByOnlineAccountIdParams{ LocalUserId };
-			TOnlineResult<FAuthGetLocalOnlineUserByOnlineAccountId> Result = Auth->GetLocalOnlineUserByOnlineAccountId(MoveTemp(GetLocalOnlineUserByOnlineAccountIdParams));
+			if (!Auth->IsLoggedIn(LocalUserId))
+			{
+				UE_LOG(LogTemp, Warning, TEXT("[FSessionsCommon::CheckLeaveSessionState] Could not leave session with user [%s] not logged in"), *ToLogString(LocalUserId));
 
-			if (Result.IsOk())
+				return Errors::InvalidUser();
+			}
+		}
+
+		return Errors::Success();
+	}
+
+	FOnlineError FSessionsCommon::CheckSendSessionInviteState(const FSendSessionInvite::Params& Params)
+	{
+		// User login check for main caller, session check
+		if (TSharedRef<FSession>* FoundSession = SessionsByName.Find(Params.SessionName))
+		{
+			if (!FoundSession->Get().SessionSettings.bIsDedicatedServerSession)
 			{
-				if (Result.GetOkValue().AccountInfo->LoginStatus == ELoginStatus::NotLoggedIn)
+				if (!Services.GetAuthInterface()->IsLoggedIn(Params.LocalUserId))
 				{
-					UE_LOG(LogTemp, Warning, TEXT("[FSessionsCommon::CheckLeaveSessionState] Could not leave session with user [%s] not logged in"), *ToLogString(LocalUserId));
+					UE_LOG(LogTemp, Warning, TEXT("[FSessionsNull::CheckSendSessionInviteState] Could not send session invite with user [%s] not logged in"), *ToLogString(Params.LocalUserId));
 
 					return Errors::InvalidUser();
 				}
 			}
 		}
+		else
+		{
+			UE_LOG(LogTemp, Warning, TEXT("[FSessionsCommon::CheckSendSessionInviteState] Could not send session invite for session with name [%s]. Session not found"), *Params.SessionName.ToString());
+
+			return Errors::NotFound();
+		}
+
+		// TODO: check if Session Invite Id is valid
+
+		return Errors::Success();
+	}
+
+	FOnlineError FSessionsCommon::CheckRejectSessionInviteState(const FRejectSessionInvite::Params& Params)
+	{
+		if (!Services.GetAuthInterface()->IsLoggedIn(Params.LocalUserId))
+		{
+			UE_LOG(LogTemp, Warning, TEXT("[FSessionsNull::CheckRejectSessionInviteState] Could not send session invite with user [%s] not logged in"), *ToLogString(Params.LocalUserId));
+
+			return Errors::InvalidUser();
+		}
+
+		// TODO: check if Session Invite Id is valid
 
 		return Errors::Success();
 	}
 
+	FOnlineError FSessionsCommon::CheckRegisterPlayersState(const FRegisterPlayers::Params& Params)
+	{
+		if (TOptional<FOnlineError> Result = CheckSessionExistsByName(Params.SessionName))
+		{
+			UE_LOG(LogTemp, Warning, TEXT("[FSessionsCommon::CheckRegisterPlayersState] Could not register players in session with name [%s]. Session not found"), *Params.SessionName.ToString());
+
+			return Result.GetValue();
+		}
+
+		// TODO: Check if there are enough slots available if Params::bReserveSlot is true
+
+		return Errors::Success();
+	}
+
+	FOnlineError FSessionsCommon::CheckUnregisterPlayersState(const FUnregisterPlayers::Params& Params)
+	{
+		if(TOptional<FOnlineError> Result = CheckSessionExistsByName(Params.SessionName))
+		{
+			UE_LOG(LogTemp, Warning, TEXT("[FSessionsCommon::CheckUnregisterPlayersState] Could not unregister players from session with name [%s]. Session not found"), *Params.SessionName.ToString());
+
+			return Result.GetValue();
+		}
+
+		return Errors::Success();
+	}
+
+	// TODO: Have all Check methods return TOptional too, change call sites, and write Macro for repeating code structure
+	TOptional<FOnlineError> FSessionsCommon::CheckSessionExistsByName(const FName& SessionName)
+	{
+		TOptional<FOnlineError> Result;
+
+		if (!SessionsByName.Contains(SessionName))
+		{
+			Result.Emplace(Errors::NotFound());
+		}
+
+		return Result;
+	}
+
 /* UE::Online */ }
\ No newline at end of file
diff --git a/Engine/Plugins/Online/Experimental/OnlineServices/Source/OnlineServicesCommon/Public/Online/SessionsCommon.h b/Engine/Plugins/Online/Experimental/OnlineServices/Source/OnlineServicesCommon/Public/Online/SessionsCommon.h
--- a/Engine/Plugins/Online/Experimental/OnlineServices/Source/OnlineServicesCommon/Public/Online/SessionsCommon.h
+++ b/Engine/Plugins/Online/Experimental/OnlineServices/Source/OnlineServicesCommon/Public/Online/SessionsCommon.h
@@ -10,6 +10,8 @@ namespace UE::Online {
 
 class FOnlineServicesCommon;
 
+static FName CONNECT_STRING_TAG = TEXT("CONNECT_STRING");
+
 template<EOnlineServices OnlineServicesType>
 class TOnlineSessionIdStringRegistry : public IOnlineSessionIdRegistry
 {
@@ -81,6 +83,8 @@ public:
 	virtual TOnlineAsyncOpHandle<FFindSessions> FindSessions(FFindSessions::Params&& Params) override;
 	virtual TOnlineAsyncOpHandle<FStartMatchmaking> StartMatchmaking(FStartMatchmaking::Params&& Params) override;
 	virtual TOnlineAsyncOpHandle<FJoinSession> JoinSession(FJoinSession::Params&& Params) override;
+	virtual TOnlineAsyncOpHandle<FAddSessionMembers> AddSessionMembers(FAddSessionMembers::Params&& Params) override;
+	virtual TOnlineAsyncOpHandle<FRemoveSessionMembers> RemoveSessionMembers(FRemoveSessionMembers::Params&& Params) override;
 	virtual TOnlineAsyncOpHandle<FSendSessionInvite> SendSessionInvite(FSendSessionInvite::Params&& Params) override;
 	virtual TOnlineAsyncOpHandle<FRejectSessionInvite> RejectSessionInvite(FRejectSessionInvite::Params&& Params) override;
 	virtual TOnlineAsyncOpHandle<FRegisterPlayers> RegisterPlayers(FRegisterPlayers::Params&& Params) override;
@@ -93,6 +97,12 @@ public:
 	virtual TOnlineEvent<void(const FUISessionJoinRequested&)> OnUISessionJoinRequested() override;
 
 protected:
+	TOnlineResult<FAddSessionMembers> AddSessionMembersImpl(const FAddSessionMembers::Params& Params);
+	TOnlineResult<FRemoveSessionMembers> RemoveSessionMembersImpl(const FRemoveSessionMembers::Params& Params);
+
+	TOnlineResult<FRegisterPlayers> RegisterPlayersImpl(const FRegisterPlayers::Params& Params);
+	TOnlineResult<FUnregisterPlayers> UnregisterPlayersImpl(const FUnregisterPlayers::Params& Params);
+
 	FOnlineError CheckCreateSessionParams(const FCreateSession::Params& Params);
 	FOnlineError CheckCreateSessionState(const FCreateSession::Params& Params);
 
@@ -104,8 +114,23 @@ protected:
 	FOnlineError CheckJoinSessionParams(const FJoinSession::Params& Params);
 	FOnlineError CheckJoinSessionState(const FJoinSession::Params& Params);
 
+	FOnlineError CheckAddSessionMembersState(const FAddSessionMembers::Params& Params);
+
+	FOnlineError CheckRemoveSessionMembersState(const FRemoveSessionMembers::Params& Params);
+
 	FOnlineError CheckLeaveSessionState(const FLeaveSession::Params& Params);
 
+	FOnlineError CheckSendSessionInviteState(const FSendSessionInvite::Params& Params);
+
+	FOnlineError CheckRejectSessionInviteState(const FRejectSessionInvite::Params& Params);
+
+	FOnlineError CheckRegisterPlayersState(const FRegisterPlayers::Params& Params);
+
+	FOnlineError CheckUnregisterPlayersState(const FUnregisterPlayers::Params& Params);
+
+private:
+	TOptional<FOnlineError> CheckSessionExistsByName(const FName& SessionName);
+
 protected:
 
 	struct FSessionEvents
diff --git a/Engine/Plugins/Online/Experimental/OnlineServices/Source/OnlineServicesCommonEngineUtils/Private/Online/SessionsLAN.cpp b/Engine/Plugins/Online/Experimental/OnlineServices/Source/OnlineServicesCommonEngineUtils/Private/Online/SessionsLAN.cpp
--- a/Engine/Plugins/Online/Experimental/OnlineServices/Source/OnlineServicesCommonEngineUtils/Private/Online/SessionsLAN.cpp
+++ b/Engine/Plugins/Online/Experimental/OnlineServices/Source/OnlineServicesCommonEngineUtils/Private/Online/SessionsLAN.cpp
@@ -10,18 +10,24 @@
 
 namespace UE::Online {
 
-/** FSessionLAN */
+/** FOnlineSessionIdRegistryLAN */
 
-FSessionLAN::FSessionLAN()
+FOnlineSessionIdRegistryLAN& FOnlineSessionIdRegistryLAN::Get()
 {
-	Initialize();
+	static FOnlineSessionIdRegistryLAN Instance;
+	return Instance;
 }
 
-FSessionLAN::FSessionLAN(const FSessionLAN& InSession)
-	: FSession(InSession)
-	, OwnerInternetAddr(InSession.OwnerInternetAddr)
+FOnlineSessionIdHandle FOnlineSessionIdRegistryLAN::GetNextSessionId()
 {
+	return BasicRegistry.FindOrAddHandle(FGuid::NewGuid().ToString());
+}
 
+/** FSessionLAN */
+
+FSessionLAN::FSessionLAN()
+{
+	Initialize();
 }
 
 void FSessionLAN::Initialize()
@@ -54,16 +60,12 @@ void FSessionLAN::Initialize()
 
 FSessionLAN& FSessionLAN::Cast(FSession& InSession)
 {
-	check(InSession.SessionId.GetOnlineServicesType() == EOnlineServices::Null);
-
-	return *static_cast<FSessionLAN*>(&InSession);
+	return static_cast<FSessionLAN&>(InSession);
 }
 
 const FSessionLAN& FSessionLAN::Cast(const FSession& InSession)
 {
-	check(InSession.SessionId.GetOnlineServicesType() == EOnlineServices::Null);
-
-	return *static_cast<const FSessionLAN*>(&InSession);
+	return static_cast<const FSessionLAN&>(InSession);
 }
 
 /** FSessionsLAN */
@@ -85,8 +87,236 @@ void FSessionsLAN::Tick(float DeltaSeconds)
 	Super::Tick(DeltaSeconds);
 }
 
-bool FSessionsLAN::TryHostLANSession()
+TOnlineAsyncOpHandle<FCreateSession> FSessionsLAN::CreateSession(FCreateSession::Params&& Params)
+{
+	TOnlineAsyncOpRef<FCreateSession> Op = GetOp<FCreateSession>(MoveTemp(Params));
+	const FCreateSession::Params& OpParams = Op->GetParams();
+
+	FOnlineError ParamsCheck = CheckCreateSessionParams(OpParams);
+	if (ParamsCheck != Errors::Success())
+	{
+		Op->SetError(MoveTemp(ParamsCheck));
+		return Op->GetHandle();
+	}
+
+	// We'll only host on the LAN beacon for public sessions
+	if (OpParams.SessionSettings.JoinPolicy == ESessionJoinPolicy::Public)
+	{
+		if (TOptional<FOnlineError> TryHostLANSessionResult = TryHostLANSession())
+		{
+			Op->SetError(MoveTemp(TryHostLANSessionResult.GetValue()));
+			return Op->GetHandle();
+		}
+	}
+
+	Op->Then([this](TOnlineAsyncOp<FCreateSession>& Op)
+	{
+		const FCreateSession::Params& OpParams = Op.GetParams();
+
+		FOnlineError StateCheck = CheckCreateSessionState(OpParams);
+		if (StateCheck != Errors::Success())
+		{
+			Op.SetError(MoveTemp(StateCheck));
+			return;
+		}
+
+		TSharedRef<FSessionLAN> NewSessionLANRef = MakeShared<FSessionLAN>();
+		NewSessionLANRef->CurrentState = ESessionState::Valid;
+		NewSessionLANRef->OwnerUserId = OpParams.LocalUserId;
+		NewSessionLANRef->SessionId = FOnlineSessionIdRegistryLAN::Get().GetNextSessionId();
+		NewSessionLANRef->SessionSettings = OpParams.SessionSettings;
+
+		// For LAN sessions, we'll add all the Session members manually instead of calling JoinSession since there is no API calls involved
+		NewSessionLANRef->SessionSettings.SessionMembers.Append(OpParams.LocalUsers);
+
+		SessionsByName.Add(OpParams.SessionName, NewSessionLANRef);
+		SessionsById.Add(NewSessionLANRef->SessionId, NewSessionLANRef);
+
+		Op.SetResult(FCreateSession::Result{ NewSessionLANRef });
+	})
+	.Enqueue(GetSerialQueue());
+
+	return Op->GetHandle();
+}
+
+TOnlineAsyncOpHandle<FUpdateSession> FSessionsLAN::UpdateSession(FUpdateSession::Params&& Params)
+{
+	TOnlineAsyncOpRef<FUpdateSession> Op = GetOp<FUpdateSession>(MoveTemp(Params));
+
+	Op->Then([this](TOnlineAsyncOp<FUpdateSession>& Op) mutable
+	{
+		const FUpdateSession::Params& OpParams = Op.GetParams();
+
+		FOnlineError StateCheck = CheckUpdateSessionState(OpParams);
+		if (StateCheck != Errors::Success())
+		{
+			Op.SetError(MoveTemp(StateCheck));
+			return;
+		}
+
+		// We'll check and update all settings one by one, with additional logic wherever required
+
+		TSharedRef<FSession>& FoundSession = SessionsByName.FindChecked(OpParams.SessionName);
+		FSessionSettings& SessionSettings = FoundSession->SessionSettings;
+		const FSessionSettingsUpdate& UpdatedSettings = OpParams.Mutations;
+
+		if (UpdatedSettings.JoinPolicy.IsSet())
+		{
+			// Changes in the join policy setting will mean we start or stop the LAN Beacon broadcasting the session information
+			// There is no FriendsLAN interface at the time of this implementation, hence the binary behavior (Public/Non-Public)
+			if (SessionSettings.JoinPolicy != ESessionJoinPolicy::Public && UpdatedSettings.JoinPolicy.GetValue() == ESessionJoinPolicy::Public)
+			{
+				if (TOptional<FOnlineError> TryHostLANSessionResult = TryHostLANSession())
+				{
+					Op.SetError(MoveTemp(TryHostLANSessionResult.GetValue()));
+					return;
+				}
+			}
+			else if (SessionSettings.JoinPolicy == ESessionJoinPolicy::Public && UpdatedSettings.JoinPolicy.GetValue() != ESessionJoinPolicy::Public)
+			{
+				StopLANSession();
+			}
+		}
+
+		SessionSettings += UpdatedSettings;
+
+		// We set the result and fire the event
+		Op.SetResult(FUpdateSession::Result{ FoundSession });
+
+		FSessionUpdated SessionUpdatedEvent{ FoundSession, UpdatedSettings };
+		SessionEvents.OnSessionUpdated.Broadcast(SessionUpdatedEvent);
+	})
+	.Enqueue(GetSerialQueue());
+
+	return Op->GetHandle();
+}
+
+TOnlineAsyncOpHandle<FFindSessions> FSessionsLAN::FindSessions(FFindSessions::Params&& Params)
+{
+	TOnlineAsyncOpRef<FFindSessions> Op = GetOp<FFindSessions>(MoveTemp(Params));
+	const FFindSessions::Params& OpParams = Op->GetParams();
+
+	FOnlineError ParamsCheck = CheckFindSessionsParams(OpParams);
+	if (ParamsCheck != Errors::Success())
+	{
+		Op->SetError(MoveTemp(ParamsCheck));
+		return Op->GetHandle();
+	}
+
+	Op->Then([this](TOnlineAsyncOp<FFindSessions>& Op) mutable
+	{
+		const FFindSessions::Params& OpParams = Op.GetParams();
+
+		FOnlineError StateCheck = CheckFindSessionsState(OpParams);
+		if (StateCheck != Errors::Success())
+		{
+			Op.SetError(MoveTemp(StateCheck));
+			return;
+		}
+
+		CurrentSessionSearch = MakeShared<FFindSessions::Result>();
+		CurrentSessionSearchHandle = Op.AsShared();
+
+		FindLANSessions();
+	})
+	.Enqueue(GetSerialQueue());
+
+	return Op->GetHandle();
+}
+
+TOnlineAsyncOpHandle<FJoinSession> FSessionsLAN::JoinSession(FJoinSession::Params&& Params)
+{
+	TOnlineAsyncOpRef<FJoinSession> Op = GetOp<FJoinSession>(MoveTemp(Params));
+	const FJoinSession::Params& OpParams = Op->GetParams();
+
+	FOnlineError ParamsCheck = CheckJoinSessionParams(OpParams);
+	if (ParamsCheck != Errors::Success())
+	{
+		Op->SetError(MoveTemp(ParamsCheck));
+		return Op->GetHandle();
+	}
+
+	// If no restrictions apply, we'll save our copy of the session, add the new players, and register it
+	Op->Then([this](TOnlineAsyncOp<FJoinSession>& Op) mutable
+	{
+		const FJoinSession::Params& OpParams = Op.GetParams();
+
+		FOnlineError StateCheck = CheckJoinSessionState(OpParams);
+		if (StateCheck != Errors::Success())
+		{
+			Op.SetError(MoveTemp(StateCheck));
+			return;
+		}
+
+		// TODO: This should retrieve an FSession from a cache and move it, instead of making a copy
+		TSharedRef<FSession> NewSessionRef = MakeShared<FSession>(*OpParams.Session);
+
+		SessionsByName.Add(OpParams.SessionName, NewSessionRef);
+		SessionsById.Add(NewSessionRef->SessionId, NewSessionRef);
+
+		Op.SetResult(FJoinSession::Result{ NewSessionRef });
+
+		TArray<FOnlineAccountIdHandle> LocalUserIds;
+		LocalUserIds.Reserve(OpParams.LocalUsers.Num());
+		OpParams.LocalUsers.GenerateKeyArray(LocalUserIds);
+
+		FSessionJoined SessionJoinedEvent = { MoveTemp(LocalUserIds), NewSessionRef };
+		
+		SessionEvents.OnSessionJoined.Broadcast(SessionJoinedEvent);
+	})
+	.Enqueue(GetSerialQueue());
+
+	return Op->GetHandle();
+}
+
+TOnlineAsyncOpHandle<FLeaveSession> FSessionsLAN::LeaveSession(FLeaveSession::Params&& Params)
 {
+	TOnlineAsyncOpRef<FLeaveSession> Op = GetOp<FLeaveSession>(MoveTemp(Params));
+
+	Op->Then([this](TOnlineAsyncOp<FLeaveSession>& Op) mutable
+	{
+		const FLeaveSession::Params& OpParams = Op.GetParams();
+
+		FOnlineError StateCheck = CheckLeaveSessionState(OpParams);
+		if (StateCheck != Errors::Success())
+		{
+			Op.SetError(MoveTemp(StateCheck));
+			return;
+		}
+
+		// This scope guarantees that FoundSession can't be used after its removal
+		{
+			TSharedRef<FSession>& FoundSession = SessionsByName.FindChecked(OpParams.SessionName);
+
+			// Only the host should stop the session at this point
+			if (FoundSession->OwnerUserId == OpParams.LocalUserId && FoundSession->SessionSettings.JoinPolicy == ESessionJoinPolicy::Public)
+			{
+				StopLANSession();
+			}
+
+			// TODO: Remove all local users from session
+
+			SessionsById.Remove(FoundSession->SessionId);
+			SessionsByName.Remove(OpParams.SessionName);
+		}
+
+		Op.SetResult(FLeaveSession::Result{ });
+
+		FSessionLeft SessionLeftEvent;
+		SessionLeftEvent.LocalUserIds.Append(OpParams.LocalUsers);
+		SessionEvents.OnSessionLeft.Broadcast(SessionLeftEvent);
+	})
+	.Enqueue(GetSerialQueue());
+
+	return Op->GetHandle();
+}
+
+/** LANSessionManager methods */
+
+TOptional<FOnlineError> FSessionsLAN::TryHostLANSession()
+{
+	TOptional<FOnlineError> Result;
+
 	// The LAN Beacon can only broadcast one session at a time
 	if (LANSessionManager->GetBeaconState() == ELanBeaconState::NotUsingLanBeacon)
 	{
@@ -97,22 +327,24 @@ bool FSessionsLAN::TryHostLANSession()
 			UE_LOG(LogTemp, VeryVerbose, TEXT("[FSessionsLAN::TryHostLANSession] LAN Beacon hosting..."));
 
 			PublicSessionsHosted++;
-
-			return true;
 		}
 		else
 		{
 			UE_LOG(LogTemp, VeryVerbose, TEXT("[FSessionsLAN::TryHostLANSession] LAN Beacon failed to host!"));
 
 			LANSessionManager->StopLANSession();
+
+			Result.Emplace(Errors::RequestFailure()); // TODO: May need a new error type
 		}
 	}
 	else
 	{
 		UE_LOG(LogTemp, VeryVerbose, TEXT("[FSessionsLAN::TryHostLANSession] LAN Beacon already in use!"));
+
+		Result.Emplace(Errors::RequestFailure()); // TODO: May need a new error type
 	}
 
-	return false;
+	return Result;
 }
 
 void FSessionsLAN::FindLANSessions()
@@ -168,9 +400,9 @@ void FSessionsLAN::StopLANSession()
 void FSessionsLAN::OnValidQueryPacketReceived(uint8* PacketData, int32 PacketLength, uint64 ClientNonce)
 {
 	// Iterate through all registered sessions and respond for each one that can be joinable
-	for (TMap<FName, TSharedRef<FSession>>::TConstIterator It(SessionsByName); It; ++It)
+	for (const TPair<FName, TSharedRef<FSession>>& Entry : SessionsByName)
 	{
-		const TSharedRef<FSession>& Session = It.Value();
+		const TSharedRef<FSession>& Session = Entry.Value;
 
 		if (Session->SessionSettings.JoinPolicy == ESessionJoinPolicy::Public)
 		{
@@ -201,7 +433,7 @@ void FSessionsLAN::OnValidResponsePacketReceived(uint8* PacketData, int32 Packet
 		TSharedRef<FSessionLAN> Session = MakeShared<FSessionLAN>();
 
 		FNboSerializeFromBuffer Packet(PacketData, PacketLength);
-		ReadSessionFromPacket(Packet, FSessionLAN::Cast(*Session));
+		ReadSessionFromPacket(Packet, *Session);
 
 		CurrentSessionSearch->FoundSessions.Add(Session);
 	}
diff --git a/Engine/Plugins/Online/Experimental/OnlineServices/Source/OnlineServicesCommonEngineUtils/Public/Online/SessionsLAN.h b/Engine/Plugins/Online/Experimental/OnlineServices/Source/OnlineServicesCommonEngineUtils/Public/Online/SessionsLAN.h
--- a/Engine/Plugins/Online/Experimental/OnlineServices/Source/OnlineServicesCommonEngineUtils/Public/Online/SessionsLAN.h
+++ b/Engine/Plugins/Online/Experimental/OnlineServices/Source/OnlineServicesCommonEngineUtils/Public/Online/SessionsLAN.h
@@ -10,11 +10,19 @@ namespace UE::Online {
 
 class FOnlineServicesCommon;
 
+class ONLINESERVICESCOMMONENGINEUTILS_API FOnlineSessionIdRegistryLAN : public TOnlineSessionIdStringRegistry<EOnlineServices::Default>
+{
+public:
+	static FOnlineSessionIdRegistryLAN& Get();
+
+	FOnlineSessionIdHandle GetNextSessionId();
+};
+
 class ONLINESERVICESCOMMONENGINEUTILS_API FSessionLAN : public FSession
 {
 public:
 	FSessionLAN();
-	FSessionLAN(const FSessionLAN& InSession);
+	FSessionLAN(const FSessionLAN& InSession) = default;
 
 	static FSessionLAN& Cast(FSession& InSession);
 	static const FSessionLAN& Cast(const FSession& InSession);
@@ -37,8 +45,15 @@ public:
 
 	virtual void Tick(float DeltaSeconds) override;
 
+	// ISessions
+	virtual TOnlineAsyncOpHandle<FCreateSession> CreateSession(FCreateSession::Params&& Params) override;
+	virtual TOnlineAsyncOpHandle<FUpdateSession> UpdateSession(FUpdateSession::Params&& Params) override;
+	virtual TOnlineAsyncOpHandle<FFindSessions> FindSessions(FFindSessions::Params&& Params) override;
+	virtual TOnlineAsyncOpHandle<FJoinSession> JoinSession(FJoinSession::Params&& Params) override;
+	virtual TOnlineAsyncOpHandle<FLeaveSession> LeaveSession(FLeaveSession::Params&& Params) override;
+
 protected:
-	bool TryHostLANSession();
+	TOptional<FOnlineError> TryHostLANSession();
 	void FindLANSessions();
 	void StopLANSession();
 	void OnValidQueryPacketReceived(uint8* PacketData, int32 PacketLength, uint64 ClientNonce);
diff --git a/Engine/Plugins/Online/Experimental/OnlineServices/Source/OnlineServicesInterface/Private/Online/Sessions.cpp b/Engine/Plugins/Online/Experimental/OnlineServices/Source/OnlineServicesInterface/Private/Online/Sessions.cpp
--- a/Engine/Plugins/Online/Experimental/OnlineServices/Source/OnlineServicesInterface/Private/Online/Sessions.cpp
+++ b/Engine/Plugins/Online/Experimental/OnlineServices/Source/OnlineServicesInterface/Private/Online/Sessions.cpp
@@ -172,98 +172,109 @@ void LexFromString(ESessionJoinPolicy& Value, const TCHAR* InStr)
 	}
 }
 
-FSessionMemberUpdate& FSessionMemberUpdate::operator+=(FSessionMemberUpdate&& UpdatedValue)
-{
-	for (TPair<FName, FCustomSessionSetting>& UpdatedMemberSetting : UpdatedValue.UpdatedMemberSettings)
-	{
-		// If an update adds a modification to a setting that had previously been marked for removal, we'll keep the latest change
-		RemovedMemberSettings.Remove(UpdatedMemberSetting.Key);
-	}
-	UpdatedMemberSettings.Append(MoveTemp(UpdatedValue.UpdatedMemberSettings));
+#define COPY_TOPTIONAL_VALUE_IF_SET(Value) \
+	if (UpdatedSettings.Value.IsSet()) \
+	{ \
+		Value = UpdatedSettings.Value.GetValue(); \
+	} \
 
-	for (FName& Key : UpdatedValue.RemovedMemberSettings)
-	{
-		// If an update removes a setting that had previously been modified, we'll keep the latest change
-		UpdatedMemberSettings.Remove(Key);
-		RemovedMemberSettings.AddUnique(MoveTemp(Key));
-	}
-
-	return *this;
-}
-
-FSessionSettingsUpdate& FSessionSettingsUpdate::operator+=(FSessionSettingsUpdate&& UpdatedValue)
+FSessionSettings& FSessionSettings::operator+=(const FSessionSettingsUpdate& UpdatedSettings)
 {
-	if (UpdatedValue.SchemaName.IsSet())
+	COPY_TOPTIONAL_VALUE_IF_SET(SchemaName) // TODO: We may need some additional logic for schema changes
+	COPY_TOPTIONAL_VALUE_IF_SET(NumMaxPublicConnections)
+	COPY_TOPTIONAL_VALUE_IF_SET(NumOpenPublicConnections)
+	COPY_TOPTIONAL_VALUE_IF_SET(NumMaxPrivateConnections)
+	COPY_TOPTIONAL_VALUE_IF_SET(NumOpenPrivateConnections)
+	COPY_TOPTIONAL_VALUE_IF_SET(JoinPolicy)
+	COPY_TOPTIONAL_VALUE_IF_SET(SessionIdOverride)
+	COPY_TOPTIONAL_VALUE_IF_SET(bIsDedicatedServerSession)
+	COPY_TOPTIONAL_VALUE_IF_SET(bAllowNewMembers)
+	COPY_TOPTIONAL_VALUE_IF_SET(bAllowSanctionedPlayers)
+	COPY_TOPTIONAL_VALUE_IF_SET(bAllowUnregisteredPlayers)
+	COPY_TOPTIONAL_VALUE_IF_SET(bAntiCheatProtected)
+	COPY_TOPTIONAL_VALUE_IF_SET(bPresenceEnabled)
+		
+	for (const FName& Key : UpdatedSettings.RemovedCustomSettings)
 	{
-		SchemaName = MoveTemp(UpdatedValue.SchemaName);
+		CustomSettings.Remove(Key);
 	}
 
-	if (UpdatedValue.NumMaxPublicConnections.IsSet())
-	{
-		NumMaxPublicConnections = MoveTemp(UpdatedValue.NumMaxPublicConnections);
-	}
+	CustomSettings.Append(UpdatedSettings.UpdatedCustomSettings);
 
-	if (UpdatedValue.NumOpenPublicConnections.IsSet())
+	for (const FOnlineAccountIdHandle& Key : UpdatedSettings.RemovedRegisteredPlayers)
 	{
-		NumOpenPublicConnections = MoveTemp(UpdatedValue.NumOpenPublicConnections);
+		RegisteredPlayers.Remove(Key);
 	}
 
-	if (UpdatedValue.NumMaxPrivateConnections.IsSet())
-	{
-		NumMaxPrivateConnections = MoveTemp(UpdatedValue.NumMaxPrivateConnections);
-	}
+	RegisteredPlayers.Append(UpdatedSettings.UpdatedRegisteredPlayers);
 
-	if (UpdatedValue.NumOpenPrivateConnections.IsSet())
+	for (const FOnlineAccountIdHandle& Key : UpdatedSettings.RemovedSessionMembers)
 	{
-		NumOpenPrivateConnections = MoveTemp(UpdatedValue.NumOpenPrivateConnections);
+		SessionMembers.Remove(Key);
 	}
 
-	if (UpdatedValue.JoinPolicy.IsSet())
+	for (const TPair<FOnlineAccountIdHandle, FSessionMemberUpdate>& Entry : UpdatedSettings.UpdatedSessionMembers)
 	{
-		JoinPolicy = MoveTemp(UpdatedValue.JoinPolicy);
-	}
+		if (FSessionMember* SessionMember = SessionMembers.Find(Entry.Key))
+		{
+			const FSessionMemberUpdate& SessionMemberUpdate = Entry.Value;
 
-	if (UpdatedValue.SessionIdOverride.IsSet())
-	{
-		SessionIdOverride = MoveTemp(UpdatedValue.SessionIdOverride);
-	}
+			for (const FName& Key : SessionMemberUpdate.RemovedMemberSettings)
+			{
+				SessionMember->MemberSettings.Remove(Key);
+			}
 
-	if (UpdatedValue.IsLANSession.IsSet())
-	{
-		IsLANSession = MoveTemp(UpdatedValue.IsLANSession);
+			SessionMember->MemberSettings.Append(SessionMemberUpdate.UpdatedMemberSettings);
+		}
 	}
 
-	if (UpdatedValue.IsDedicatedServerSession.IsSet())
-	{
-		IsDedicatedServerSession = MoveTemp(UpdatedValue.IsDedicatedServerSession);
-	}
+	return *this;
+}
 
-	if (UpdatedValue.bAllowNewMembers.IsSet())
-	{
-		bAllowNewMembers = MoveTemp(UpdatedValue.bAllowNewMembers);
-	}
+#undef COPY_TOPTIONAL_VALUE_IF_SET
 
-	if (UpdatedValue.bAllowSanctionedPlayers.IsSet())
+FSessionMemberUpdate& FSessionMemberUpdate::operator+=(FSessionMemberUpdate&& UpdatedValue)
+{
+	for (const TPair<FName, FCustomSessionSetting>& UpdatedMemberSetting : UpdatedValue.UpdatedMemberSettings)
 	{
-		bAllowSanctionedPlayers = MoveTemp(UpdatedValue.bAllowSanctionedPlayers);
+		// If an update adds a modification to a setting that had previously been marked for removal, we'll keep the latest change
+		RemovedMemberSettings.Remove(UpdatedMemberSetting.Key);
 	}
+	UpdatedMemberSettings.Append(MoveTemp(UpdatedValue.UpdatedMemberSettings));
 
-	if (UpdatedValue.bAllowUnregisteredPlayers.IsSet())
+	for (FName& Key : UpdatedValue.RemovedMemberSettings)
 	{
-		bAllowUnregisteredPlayers = MoveTemp(UpdatedValue.bAllowUnregisteredPlayers);
+		// If an update removes a setting that had previously been modified, we'll keep the latest change
+		UpdatedMemberSettings.Remove(Key);
+		RemovedMemberSettings.AddUnique(MoveTemp(Key));
 	}
 
-	if (UpdatedValue.bAntiCheatProtected.IsSet())
-	{
-		bAntiCheatProtected = MoveTemp(UpdatedValue.bAntiCheatProtected);
-	}
+	return *this;
+}
 
-	if (UpdatedValue.bPresenceEnabled.IsSet())
-	{
-		UpdatedValue.bPresenceEnabled = MoveTemp(UpdatedValue.bPresenceEnabled);
-	}
+#define MOVE_TOPTIONAL_IF_SET(Value) \
+	if (UpdatedValue.Value.IsSet()) \
+	{ \
+		Value = MoveTemp(UpdatedValue.Value); \
+	} \
 
-	for (TPair<FName, FCustomSessionSetting>& UpdatedCustomSetting : UpdatedValue.UpdatedCustomSettings)
+FSessionSettingsUpdate& FSessionSettingsUpdate::operator+=(FSessionSettingsUpdate&& UpdatedValue)
+{
+	MOVE_TOPTIONAL_IF_SET(SchemaName)
+	MOVE_TOPTIONAL_IF_SET(NumMaxPublicConnections)
+	MOVE_TOPTIONAL_IF_SET(NumOpenPublicConnections)
+	MOVE_TOPTIONAL_IF_SET(NumMaxPrivateConnections)
+	MOVE_TOPTIONAL_IF_SET(NumOpenPrivateConnections)
+	MOVE_TOPTIONAL_IF_SET(JoinPolicy)
+	MOVE_TOPTIONAL_IF_SET(SessionIdOverride)
+	MOVE_TOPTIONAL_IF_SET(bIsDedicatedServerSession)
+	MOVE_TOPTIONAL_IF_SET(bAllowNewMembers)
+	MOVE_TOPTIONAL_IF_SET(bAllowSanctionedPlayers)
+	MOVE_TOPTIONAL_IF_SET(bAllowUnregisteredPlayers)
+	MOVE_TOPTIONAL_IF_SET(bAntiCheatProtected)
+	MOVE_TOPTIONAL_IF_SET(bPresenceEnabled)
+
+	for (const TPair<FName, FCustomSessionSetting>& UpdatedCustomSetting : UpdatedValue.UpdatedCustomSettings)
 	{
 		// If an update adds a modification to a setting that had previously been marked for removal, we'll keep the latest change
 		RemovedCustomSettings.Remove(UpdatedCustomSetting.Key);
@@ -293,23 +304,25 @@ FSessionSettingsUpdate& FSessionSettingsUpdate::operator+=(FSessionSettingsUpdat
 		RemovedSessionMembers.AddUnique(MoveTemp(Key));
 	}
 
-	for (TPair<FOnlineAccountIdHandle, FRegisteredUser>& UpdatedRegisteredUser : UpdatedValue.UpdatedRegisteredUsers)
+	for (const TPair<FOnlineAccountIdHandle, FRegisteredPlayer>& UpdatedRegisteredPlayer : UpdatedValue.UpdatedRegisteredPlayers)
 	{
-		// If an update adds a modification to a registered user that had previously been marked for removal, we'll keep the latest change
-		RemovedRegisteredUsers.Remove(UpdatedRegisteredUser.Key);
+		// If an update adds a modification to a registered player that had previously been marked for removal, we'll keep the latest change
+		RemovedRegisteredPlayers.Remove(UpdatedRegisteredPlayer.Key);
 	}
-	UpdatedRegisteredUsers.Append(MoveTemp(UpdatedValue.UpdatedRegisteredUsers));
+	UpdatedRegisteredPlayers.Append(MoveTemp(UpdatedValue.UpdatedRegisteredPlayers));
 
-	for (FOnlineAccountIdHandle& Key : UpdatedValue.RemovedRegisteredUsers)
+	for (FOnlineAccountIdHandle& Key : UpdatedValue.RemovedRegisteredPlayers)
 	{
-		// If an update removes a registered user that had previously been modified, we'll keep the latest change
-		UpdatedRegisteredUsers.Remove(Key);
-		RemovedRegisteredUsers.AddUnique(MoveTemp(Key));
+		// If an update removes a registered player that had previously been modified, we'll keep the latest change
+		UpdatedRegisteredPlayers.Remove(Key);
+		RemovedRegisteredPlayers.AddUnique(MoveTemp(Key));
 	}
 
 	return *this;
 }
 
+#undef MOVE_TOPTIONAL_IF_SET
+
 FSession::FSession()
 {
 
diff --git a/Engine/Plugins/Online/Experimental/OnlineServices/Source/OnlineServicesInterface/Public/Online/Sessions.h b/Engine/Plugins/Online/Experimental/OnlineServices/Source/OnlineServicesInterface/Public/Online/Sessions.h
--- a/Engine/Plugins/Online/Experimental/OnlineServices/Source/OnlineServicesInterface/Public/Online/Sessions.h
+++ b/Engine/Plugins/Online/Experimental/OnlineServices/Source/OnlineServicesInterface/Public/Online/Sessions.h
@@ -70,6 +70,7 @@ struct FCustomSessionSetting
 
 using FCustomSessionSettingsMap = TMap<FName, FCustomSessionSetting>;
 
+/** A member is a player that is part of the session, and it stops being a member when they leave it */
 struct FSessionMember
 {
 	FCustomSessionSettingsMap MemberSettings;
@@ -108,33 +109,81 @@ enum class ESessionJoinPolicy : uint8
 ONLINESERVICESINTERFACE_API const TCHAR* LexToString(ESessionJoinPolicy Value);
 ONLINESERVICESINTERFACE_API void LexFromString(ESessionJoinPolicy& Value, const TCHAR* InStr);
 
-struct FRegisteredUser
+/** A player registered with the session, whether they are in it or not */
+struct FRegisteredPlayer
 {
-	/* Whether a slot was reserved for this user upon registration */
+	/* Whether a slot was reserved for this player upon registration */
 	bool bHasReservedSlot = false;
 
-	/* Whether the user is currently in the session */
+	/* Whether the player is currently in the session */
 	bool bIsInSession = false;
 };
 
-using FRegisteredUsersMap = TMap<FOnlineAccountIdHandle, FRegisteredUser>;
+using FRegisteredPlayersMap = TMap<FOnlineAccountIdHandle, FRegisteredPlayer>;
 
-struct FSessionSettings
+struct FSessionSettingsUpdate
+{
+	/** Set with an updated value if the SchemaName field will be changed in the update operation */
+	TOptional<FName> SchemaName;
+	/** Set with an updated value if the NumMaxPublicConnections field will be changed in the update operation */
+	TOptional<uint32> NumMaxPublicConnections;
+	/** Set with an updated value if the NumOpenPublicConnections field will be changed in the update operation */
+	TOptional<uint32> NumOpenPublicConnections;
+	/** Set with an updated value if the NumMaxPrivateConnections field will be changed in the update operation */
+	TOptional<uint32> NumMaxPrivateConnections;
+	/** Set with an updated value if the NumOpenPrivateConnections field will be changed in the update operation */
+	TOptional<uint32> NumOpenPrivateConnections;
+	/** Set with an updated value if the JoinPolicy field will be changed in the update operation */
+	TOptional<ESessionJoinPolicy> JoinPolicy;
+	/** Set with an updated value if the SessionIdOverride field will be changed in the update operation */
+	TOptional<FString> SessionIdOverride;
+	/** Set with an updated value if the bIsDedicatedServerSession field will be changed in the update operation */
+	TOptional<bool> bIsDedicatedServerSession;
+	/** Set with an updated value if the bAllowNewMembers field will be changed in the update operation */
+	TOptional<bool> bAllowNewMembers;
+	/** Set with an updated value if the bAllowSanctionedPlayers field will be changed in the update operation */
+	TOptional<bool> bAllowSanctionedPlayers;
+	/** Set with an updated value if the bAllowUnregisteredPlayers field will be changed in the update operation */
+	TOptional<bool> bAllowUnregisteredPlayers;
+	/** Set with an updated value if the bAntiCheatProtected field will be changed in the update operation */
+	TOptional<bool> bAntiCheatProtected;
+	/** Set with an updated value if the bPresenceEnabled field will be changed in the update operation */
+	TOptional<bool> bPresenceEnabled;
+
+	/** Updated values for custom settings to change in the update operation*/
+	FCustomSessionSettingsMap UpdatedCustomSettings;
+	/** Names of custom settings to be removed in the update operation*/
+	TArray<FName> RemovedCustomSettings;
+
+	/** Updated values for session member info to change in the update operation*/
+	FSessionMemberUpdatesMap UpdatedSessionMembers;
+	/** Id handles for session members to be removed in the update operation*/
+	TArray<FOnlineAccountIdHandle> RemovedSessionMembers;
+
+	/** Updated values for registered players to change in the update operation*/
+	FRegisteredPlayersMap UpdatedRegisteredPlayers;
+	/** Id handles for registered players to be removed in the update operation*/
+	TArray<FOnlineAccountIdHandle> RemovedRegisteredPlayers;
+
+	FSessionSettingsUpdate& operator+=(FSessionSettingsUpdate&& UpdatedValue);
+};
+
+struct ONLINESERVICESINTERFACE_API FSessionSettings
 {
 	/* The schema which will be applied to the session */
 	FName SchemaName;
 
 	/* Maximum number of public slots for session members */
-	uint32 NumMaxPublicConnections;
+	uint32 NumMaxPublicConnections = 0;
 
 	/* Number of available public slots for session members */
-	uint32 NumOpenPublicConnections;
+	uint32 NumOpenPublicConnections = 0;
 
 	/* Maximum number of private slots for session members */
-	uint32 NumMaxPrivateConnections;
+	uint32 NumMaxPrivateConnections = 0;
 
 	/* Number of available private slots for session members */
-	uint32 NumOpenPrivateConnections;
+	uint32 NumOpenPrivateConnections = 0;
 
 	/* Enum value describing the level of restriction to join the session */
 	ESessionJoinPolicy JoinPolicy;
@@ -142,26 +191,26 @@ struct FSessionSettings
  	/* In platforms that support this feature, it will set the session id to this value. Might be subject to minimum and maximum length */
  	FString SessionIdOverride;
 
-	/* Whether the session is only available in the local network and not via internet connection. Only available in some platforms */
-	bool IsLANSession;
+	/* Whether the session is only available in the local network and not via internet connection. Only available in some platforms. False by default */
+	bool bIsLANSession = false;
 
-	/* Whether the session is configured to run as a dedicated server. Only available in some platforms */
-	bool IsDedicatedServerSession;
+	/* Whether the session is configured to run as a dedicated server. Only available in some platforms. False by default */
+	bool bIsDedicatedServerSession = false;
 
-	/* Whether new members are accepted in the session. Can vary depending on various factors, like the number of free slots available, or Join-In-Progress preferences when the session has started */
-	bool bAllowNewMembers;
+	/* Whether players (registered or not) are accepted as new members in the session. Can vary depending on various factors, like the number of free slots available, or Join-In-Progress preferences when the session has started. True by default */
+	bool bAllowNewMembers = true;
 
-	/* Whether this session will allow sanctioned players to join it */
-	bool bAllowSanctionedPlayers;
+	/* Whether this session will allow sanctioned players to join it. True by default */
+	bool bAllowSanctionedPlayers = true;
 
-	/* Whether this session will allow unregistered players to join it */
-	bool bAllowUnregisteredPlayers;
+	/* Whether this session will allow unregistered players to join it. True by default */
+	bool bAllowUnregisteredPlayers = true;
 
-	/*Whether this is a secure session protected by anti-cheat services */
-	bool bAntiCheatProtected;
+	/*Whether this is a secure session protected by anti-cheat services. False by default */
+	bool bAntiCheatProtected = false;
 
-	/* Whether this session will show its information in presence updates. Can only be set in one session at a time */
-	bool bPresenceEnabled;
+	/* Whether this session will show its information in presence updates. Can only be set in one session at a time. False by default */
+	bool bPresenceEnabled = false;
   
  	/* Map of user-defined settings to be passed to the platform APIs as additional information for various purposes */
  	FCustomSessionSettingsMap CustomSettings;
@@ -169,37 +218,10 @@ struct FSessionSettings
  	/* Map of session member ids to their corresponding user-defined settings */
  	FSessionMembersMap SessionMembers;
  
- 	/* Map of registered users for this session. Can only be altered via RegisterPlayers and UnregisterPlayers */
- 	FRegisteredUsersMap RegisteredUsers;
-};
+ 	/* Map of registered players for this session. Can only be altered via RegisterPlayers and UnregisterPlayers */
+ 	FRegisteredPlayersMap RegisteredPlayers;
 
-struct FSessionSettingsUpdate
-{
-	TOptional<FName> SchemaName;
-	TOptional<uint32> NumMaxPublicConnections;
-	TOptional<uint32> NumOpenPublicConnections;
-	TOptional<uint32> NumMaxPrivateConnections;
-	TOptional<uint32> NumOpenPrivateConnections;
-	TOptional<ESessionJoinPolicy> JoinPolicy;
-	TOptional<FString> SessionIdOverride;
-	TOptional<bool> IsLANSession;
-	TOptional<bool> IsDedicatedServerSession;
-	TOptional<bool> bAllowNewMembers;
-	TOptional<bool> bAllowSanctionedPlayers;
-	TOptional<bool> bAllowUnregisteredPlayers;
-	TOptional<bool> bAntiCheatProtected;
-	TOptional<bool> bPresenceEnabled;
-
-	FCustomSessionSettingsMap UpdatedCustomSettings;
-	TArray<FName> RemovedCustomSettings;
-
-	FSessionMemberUpdatesMap UpdatedSessionMembers;
-	TArray<FOnlineAccountIdHandle> RemovedSessionMembers;
-
-	FRegisteredUsersMap UpdatedRegisteredUsers;
-	TArray<FOnlineAccountIdHandle> RemovedRegisteredUsers;
-
-	FSessionSettingsUpdate& operator+=(FSessionSettingsUpdate&& UpdatedValue);
+	FSessionSettings& operator+=(const FSessionSettingsUpdate& UpdatedValue);
 };
 
 struct ONLINESERVICESINTERFACE_API FSession
@@ -374,6 +396,9 @@ struct FFindSessions
 		/* Maximum number of results to return in one search */
 		uint32 MaxResults;
 
+		/** Whether we want to look for LAN sessions or Online sessions */
+		bool bFindLANSessions;
+
 		/* Filters to apply when searching for sessions. */
 		TArray<FFindSessionsSearchFilter> Filters;
 
@@ -435,7 +460,7 @@ struct FJoinSession
 		/* A reference to the session to be joined. */
 		TSharedRef<const FSession> Session; // TODO: FOnlineSessionIdHandle to be used after we cache search results
 
-		/* Information for all local users who will join the session (includes the session creator)*/
+		/* Information for all local users who will join the session (includes the session creator). Any player that wants to join the session needs defined information to become a new member */
 		FSessionMembersMap LocalUsers;
 	};
 
@@ -447,6 +472,56 @@ struct FJoinSession
 	};
 };
 
+struct FAddSessionMembers
+{
+	static constexpr TCHAR Name[] = TEXT("AddSessionMembers");
+
+	struct Params
+	{
+		/* The local user agent */
+		FOnlineAccountIdHandle LocalUserId;
+
+		/* Local name for the session */
+		FName SessionName;
+
+		/** Information for the session members to be added to the session. Any player that joins the session becomes a new member in doing so */
+		FSessionMembersMap NewSessionMembers;
+
+		/** Whether or not the new session members should also be added to the list of registered players. True by default*/
+		bool bRegisterPlayers = true;
+	};
+
+	struct Result
+	{
+
+	};
+};
+
+struct FRemoveSessionMembers
+{
+	static constexpr TCHAR Name[] = TEXT("RemoveSessionMembers");
+
+	struct Params
+	{
+		/* The local user agent */
+		FOnlineAccountIdHandle LocalUserId;
+
+		/* Local name for the session */
+		FName SessionName;
+
+		/* Id handles for the session members to be removed from the session */
+		TArray<FOnlineAccountIdHandle> SessionMemberIds;
+
+		/** Whether or not the session members should also be removed from the list of registered players. True by default*/
+		bool bUnregisterPlayers = true;
+	};
+
+	struct Result
+	{
+
+	};
+};
+
 struct FSendSessionInvite
 {
 	static constexpr TCHAR Name[] = TEXT("SendSessionInvite");
@@ -499,7 +574,7 @@ struct FRegisterPlayers
 		/* Array of users which will be registered */
 		TArray<FOnlineAccountIdHandle> TargetUsers;
 
-		/* Whether a slot should be saved for the registered users */
+		/* Whether a slot should be saved for the registered players */
 		bool bReserveSlot;
 	};
 
@@ -670,6 +745,28 @@ public:
 	 */
 	virtual TOnlineAsyncOpHandle<FJoinSession> JoinSession(FJoinSession::Params&& Params) = 0;
 
+	/**
+	 * Adds a set of new session members to the named session
+	 * Session member information passed will be saved in the session settings
+	 * Number of open slots in the session will decrease accordingly
+	 * If indicated, players will also be registered in the session
+	 * 
+	 * @params Parameters for the AddSessionMember call
+	 * @return
+	 */
+	virtual TOnlineAsyncOpHandle<FAddSessionMembers> AddSessionMembers(FAddSessionMembers::Params&& Params) = 0;
+
+	/**
+	 * Removes a set of session member from the named session
+	 * Session member information for them will be removed from session settings
+	 * Number of open slots in the session will increase accordingly
+	 * If indicated, players will also be unregistered from the session
+	 *
+	 * @params Parameters for the RemoveSessionMember call
+	 * @return
+	 */
+	virtual TOnlineAsyncOpHandle<FRemoveSessionMembers> RemoveSessionMembers(FRemoveSessionMembers::Params&& Params) = 0;
+
 	/**
 	 * Sends an invite to the named session to all given users.
 	 *
@@ -687,7 +784,8 @@ public:
 	virtual TOnlineAsyncOpHandle<FRejectSessionInvite> RejectSessionInvite(FRejectSessionInvite::Params&& Params) = 0;
 
 	/**
-	 * Registers given players in the named session.
+	 * Registers given players in the named session
+	 * If indicated, and if any are available, a slot in the session will be reserved for them
 	 *
 	 * @param Parameters for the RegisterPlayers call
 	 * @return
@@ -696,6 +794,7 @@ public:
 
 	/**
 	 * Unregisters given players from the named session.
+	 * If indicated, and if they are members of it, players will also be removed from the session
 	 *
 	 * @param Parameters for the UnregisterPlayers call
 	 * @return
@@ -762,9 +861,9 @@ BEGIN_ONLINE_STRUCT_META(FSessionMemberUpdate)
 	ONLINE_STRUCT_FIELD(FSessionMemberUpdate, RemovedMemberSettings)
 END_ONLINE_STRUCT_META()
 
-BEGIN_ONLINE_STRUCT_META(FRegisteredUser)
-	ONLINE_STRUCT_FIELD(FRegisteredUser, bHasReservedSlot),
-	ONLINE_STRUCT_FIELD(FRegisteredUser, bIsInSession)
+BEGIN_ONLINE_STRUCT_META(FRegisteredPlayer)
+	ONLINE_STRUCT_FIELD(FRegisteredPlayer, bHasReservedSlot),
+	ONLINE_STRUCT_FIELD(FRegisteredPlayer, bIsInSession)
 END_ONLINE_STRUCT_META()
 
 BEGIN_ONLINE_STRUCT_META(FSessionSettings)
@@ -775,8 +874,8 @@ BEGIN_ONLINE_STRUCT_META(FSessionSettings)
 	ONLINE_STRUCT_FIELD(FSessionSettings, NumOpenPrivateConnections),
 	ONLINE_STRUCT_FIELD(FSessionSettings, JoinPolicy),
 	ONLINE_STRUCT_FIELD(FSessionSettings, SessionIdOverride),
-	ONLINE_STRUCT_FIELD(FSessionSettings, IsLANSession),
-	ONLINE_STRUCT_FIELD(FSessionSettings, IsDedicatedServerSession),
+	ONLINE_STRUCT_FIELD(FSessionSettings, bIsLANSession),
+	ONLINE_STRUCT_FIELD(FSessionSettings, bIsDedicatedServerSession),
 	ONLINE_STRUCT_FIELD(FSessionSettings, bAllowNewMembers),
 	ONLINE_STRUCT_FIELD(FSessionSettings, bAllowSanctionedPlayers),
 	ONLINE_STRUCT_FIELD(FSessionSettings, bAllowUnregisteredPlayers),
@@ -784,7 +883,7 @@ BEGIN_ONLINE_STRUCT_META(FSessionSettings)
 	ONLINE_STRUCT_FIELD(FSessionSettings, bPresenceEnabled),
 	ONLINE_STRUCT_FIELD(FSessionSettings, CustomSettings),
 	ONLINE_STRUCT_FIELD(FSessionSettings, SessionMembers),
-	ONLINE_STRUCT_FIELD(FSessionSettings, RegisteredUsers)
+	ONLINE_STRUCT_FIELD(FSessionSettings, RegisteredPlayers)
 END_ONLINE_STRUCT_META()
 
 BEGIN_ONLINE_STRUCT_META(FSessionSettingsUpdate)
@@ -795,8 +894,7 @@ BEGIN_ONLINE_STRUCT_META(FSessionSettingsUpdate)
 	ONLINE_STRUCT_FIELD(FSessionSettingsUpdate, NumOpenPrivateConnections),
 	ONLINE_STRUCT_FIELD(FSessionSettingsUpdate, JoinPolicy),
 	ONLINE_STRUCT_FIELD(FSessionSettingsUpdate, SessionIdOverride),
-	ONLINE_STRUCT_FIELD(FSessionSettingsUpdate, IsLANSession),
-	ONLINE_STRUCT_FIELD(FSessionSettingsUpdate, IsDedicatedServerSession),
+	ONLINE_STRUCT_FIELD(FSessionSettingsUpdate, bIsDedicatedServerSession),
 	ONLINE_STRUCT_FIELD(FSessionSettingsUpdate, bAllowNewMembers),
 	ONLINE_STRUCT_FIELD(FSessionSettingsUpdate, bAllowSanctionedPlayers),
 	ONLINE_STRUCT_FIELD(FSessionSettingsUpdate, bAllowUnregisteredPlayers),
@@ -806,8 +904,8 @@ BEGIN_ONLINE_STRUCT_META(FSessionSettingsUpdate)
 	ONLINE_STRUCT_FIELD(FSessionSettingsUpdate, RemovedCustomSettings),
 	ONLINE_STRUCT_FIELD(FSessionSettingsUpdate, UpdatedSessionMembers),
 	ONLINE_STRUCT_FIELD(FSessionSettingsUpdate, RemovedSessionMembers),
-	ONLINE_STRUCT_FIELD(FSessionSettingsUpdate, UpdatedRegisteredUsers),
-	ONLINE_STRUCT_FIELD(FSessionSettingsUpdate, RemovedRegisteredUsers)
+	ONLINE_STRUCT_FIELD(FSessionSettingsUpdate, UpdatedRegisteredPlayers),
+	ONLINE_STRUCT_FIELD(FSessionSettingsUpdate, RemovedRegisteredPlayers)
 END_ONLINE_STRUCT_META()
 
 BEGIN_ONLINE_STRUCT_META(FSession)
@@ -880,6 +978,7 @@ END_ONLINE_STRUCT_META()
 BEGIN_ONLINE_STRUCT_META(FFindSessions::Params)
 	ONLINE_STRUCT_FIELD(FFindSessions::Params, LocalUserId),
 	ONLINE_STRUCT_FIELD(FFindSessions::Params, MaxResults),
+	ONLINE_STRUCT_FIELD(FFindSessions::Params, bFindLANSessions),	
 	ONLINE_STRUCT_FIELD(FFindSessions::Params, Filters),
 	ONLINE_STRUCT_FIELD(FFindSessions::Params, TargetUser),
 	ONLINE_STRUCT_FIELD(FFindSessions::Params, SessionId)
@@ -912,6 +1011,26 @@ BEGIN_ONLINE_STRUCT_META(FJoinSession::Result)
 	ONLINE_STRUCT_FIELD(FJoinSession::Result, Session)
 END_ONLINE_STRUCT_META()
 
+BEGIN_ONLINE_STRUCT_META(FAddSessionMembers::Params)
+	ONLINE_STRUCT_FIELD(FAddSessionMembers::Params, LocalUserId),
+	ONLINE_STRUCT_FIELD(FAddSessionMembers::Params, SessionName),
+	ONLINE_STRUCT_FIELD(FAddSessionMembers::Params, NewSessionMembers),
+	ONLINE_STRUCT_FIELD(FAddSessionMembers::Params, bRegisterPlayers)
+END_ONLINE_STRUCT_META()
+
+BEGIN_ONLINE_STRUCT_META(FAddSessionMembers::Result)
+END_ONLINE_STRUCT_META()
+
+BEGIN_ONLINE_STRUCT_META(FRemoveSessionMembers::Params)
+	ONLINE_STRUCT_FIELD(FRemoveSessionMembers::Params, LocalUserId),
+	ONLINE_STRUCT_FIELD(FRemoveSessionMembers::Params, SessionName),
+	ONLINE_STRUCT_FIELD(FRemoveSessionMembers::Params, SessionMemberIds),
+	ONLINE_STRUCT_FIELD(FRemoveSessionMembers::Params, bUnregisterPlayers)
+END_ONLINE_STRUCT_META()
+
+BEGIN_ONLINE_STRUCT_META(FRemoveSessionMembers::Result)
+END_ONLINE_STRUCT_META()
+
 BEGIN_ONLINE_STRUCT_META(FSendSessionInvite::Params)
 	ONLINE_STRUCT_FIELD(FSendSessionInvite::Params, LocalUserId),
 	ONLINE_STRUCT_FIELD(FSendSessionInvite::Params, SessionName),
diff --git a/Engine/Plugins/Online/Experimental/OnlineServicesEOSGS/Source/OnlineServicesEOSGS.Build.cs b/Engine/Plugins/Online/Experimental/OnlineServicesEOSGS/Source/OnlineServicesEOSGS.Build.cs
--- a/Engine/Plugins/Online/Experimental/OnlineServicesEOSGS/Source/OnlineServicesEOSGS.Build.cs
+++ b/Engine/Plugins/Online/Experimental/OnlineServicesEOSGS/Source/OnlineServicesEOSGS.Build.cs
@@ -12,6 +12,8 @@ public class OnlineServicesEOSGS : ModuleRules
 				"Core",
 				"OnlineServicesInterface",
 				"OnlineServicesCommon",
+				"OnlineServicesCommonEngineUtils",
+				"OnlineBase",
 			}
 		);
 
diff --git a/Engine/Plugins/Online/Experimental/OnlineServicesEOSGS/Source/Private/Online/SessionsEOSGS.cpp b/Engine/Plugins/Online/Experimental/OnlineServicesEOSGS/Source/Private/Online/SessionsEOSGS.cpp
--- a/Engine/Plugins/Online/Experimental/OnlineServicesEOSGS/Source/Private/Online/SessionsEOSGS.cpp
+++ b/Engine/Plugins/Online/Experimental/OnlineServicesEOSGS/Source/Private/Online/SessionsEOSGS.cpp
@@ -6,6 +6,7 @@
 #include "Online/OnlineIdEOSGS.h"
 #include "Online/OnlineServicesEOSGS.h"
 #include "Online/OnlineServicesEOSGSTypes.h"
+#include "Online/NboSerializerEOSGSSvc.h"
 
 #include "eos_sessions.h"
 #include "eos_sessions_types.h"
@@ -37,19 +38,6 @@ FOnlineSessionIdRegistryEOSGS& FOnlineSessionIdRegistryEOSGS::Get()
 
 /** FSessionEOSGS */
 
-FSessionEOSGS::FSessionEOSGS()
-	: FSessionEOSGS(FSession())
-{
-}
-
-FSessionEOSGS::FSessionEOSGS(const FSession& InSession)
-	: FSession(InSession)
-{
-	const FSessionEOSGS& Session = Cast(InSession);
-
-	SessionDetailsHandle = Session.SessionDetailsHandle;
-}
-
 FSessionEOSGS::FSessionEOSGS(const EOS_HSessionDetails& InSessionDetailsHandle)
 	: SessionDetailsHandle(MakeShared<FSessionDetailsHandleEOSGS>(InSessionDetailsHandle))
 {
@@ -85,7 +73,6 @@ const FSessionEOSGS& FSessionEOSGS::Cast(const FSession& InSession)
 
 FSessionsEOSGS::FSessionsEOSGS(FOnlineServicesEOSGS& InServices)
 	: Super(InServices)
-	, Services(InServices)
 {
 }
 
@@ -214,50 +201,36 @@ void FSessionsEOSGS::HandleJoinSessionAccepted(const EOS_Sessions_JoinSessionAcc
 
 TOnlineAsyncOpHandle<FCreateSession> FSessionsEOSGS::CreateSession(FCreateSession::Params&& Params)
 {
+	// LAN Sessions
+	if (Params.SessionSettings.bIsLANSession)
+	{
+		return FSessionsLAN::CreateSession(MoveTemp(Params));
+	}
+
+	// EOSGS Sessions
+
 	TOnlineAsyncOpRef<FCreateSession> Op = GetOp<FCreateSession>(MoveTemp(Params));
+	const FCreateSession::Params& OpParams = Op->GetParams();
+
+	FOnlineError ParamsCheck = CheckCreateSessionParams(OpParams);
+	if (ParamsCheck != Errors::Success())
+	{
+		Op->SetError(MoveTemp(ParamsCheck));
+		return Op->GetHandle();
+	}
 
 	Op->Then([this](TOnlineAsyncOp<FCreateSession>& Op)
 	{
 		const FCreateSession::Params& OpParams = Op.GetParams();
 
-		if (TSharedRef<FSessionEOSGS>* FoundSession = SessionsByName.Find(OpParams.SessionName))
-		{
-			UE_LOG(LogTemp, Warning, TEXT("[FSessionsEOSGS::CreateSession] Could not create session with name [%s]. A session with that name already exists"), *OpParams.SessionName.ToString());
-
-			Op.SetError(Errors::InvalidState()); // TODO: New error: Session with name %s already exists
-
-			return MakeFulfilledPromise<TDefaultErrorResult<FUpdateSessionImpl>>().GetFuture();
-		}
-
-		if (OpParams.SessionSettings.NumMaxPrivateConnections == 0 && OpParams.SessionSettings.NumMaxPublicConnections == 0)
+		FOnlineError StateCheck = CheckCreateSessionState(OpParams);
+		if (StateCheck != Errors::Success())
 		{
-			UE_LOG(LogTemp, Warning, TEXT("[FSessionsEOSGS::CreateSession] Could not create session with no valid NumMaxPrivateConnections [%d] or NumMaxPublicConnections [%d]"), OpParams.SessionSettings.NumMaxPrivateConnections, OpParams.SessionSettings.NumMaxPublicConnections);
-
-			Op.SetError(Errors::InvalidParams());
+			Op.SetError(MoveTemp(StateCheck));
 
 			return MakeFulfilledPromise<TDefaultErrorResult<FUpdateSessionImpl>>().GetFuture();
 		}
 
-		if (!OpParams.SessionSettings.IsDedicatedServerSession)
-		{
-			IAuthPtr Auth = Services.GetAuthInterface();
-
-			FAuthGetLocalOnlineUserByOnlineAccountId::Params GetLocalOnlineUserByOnlineAccountIdParams { OpParams.LocalUserId };
-			TOnlineResult<FAuthGetLocalOnlineUserByOnlineAccountId> Result = Auth->GetLocalOnlineUserByOnlineAccountId(MoveTemp(GetLocalOnlineUserByOnlineAccountIdParams));
-
-			if (Result.IsOk())
-			{
-				if (Result.GetOkValue().AccountInfo->LoginStatus == ELoginStatus::NotLoggedIn)
-				{
-					UE_LOG(LogTemp, Warning, TEXT("[FSessionsEOSGS::CreateSession] Could not create session with user [%s] not logged in"), *ToLogString(OpParams.LocalUserId));
-
-					Op.SetError(Errors::InvalidUser());
-
-					return MakeFulfilledPromise<TDefaultErrorResult<FUpdateSessionImpl>>().GetFuture();
-				}
-			}
-		}
-
 		if (!OpParams.SessionSettings.SessionIdOverride.IsEmpty())
 		{
 			int32 Length = OpParams.SessionSettings.SessionIdOverride.Len();
@@ -320,17 +293,15 @@ TOnlineAsyncOpHandle<FCreateSession> FSessionsEOSGS::CreateSession(FCreateSessio
 		WriteCreateSessionModificationHandle(SessionModificationHandle, OpParams.SessionSettings);
 
 		// If the session modification handle is created successfully, we'll create the local session object
-		FSessionEOSGS NewSession = FSessionEOSGS();
-		NewSession.CurrentState = ESessionState::Creating;
-		NewSession.OwnerUserId = OpParams.LocalUserId;
-		NewSession.SessionSettings = OpParams.SessionSettings;
-
- 		TSharedRef<FSessionEOSGS> NewSessionEOSGSRef = MakeShared<FSessionEOSGS>(NewSession);
+ 		TSharedRef<FSessionEOSGS> NewSessionEOSGSRef = MakeShared<FSessionEOSGS>();
+		NewSessionEOSGSRef->CurrentState = ESessionState::Creating;
+		NewSessionEOSGSRef->OwnerUserId = OpParams.LocalUserId;
+		NewSessionEOSGSRef->SessionSettings = OpParams.SessionSettings;
 
  		SessionsByName.Add(OpParams.SessionName, NewSessionEOSGSRef);
- 		SessionsById.Add(NewSession.SessionId, NewSessionEOSGSRef);
+ 		SessionsById.Add(NewSessionEOSGSRef->SessionId, NewSessionEOSGSRef);
 
-		FUpdateSessionImpl::Params UpdateSessionImplParams { SessionModificationHandle };
+		FUpdateSessionImpl::Params UpdateSessionImplParams { MakeShared<FSessionModificationHandleEOSGS>(SessionModificationHandle) };
 		
 		return UpdateSessionImpl(MoveTemp(UpdateSessionImplParams));
 	})
@@ -340,27 +311,18 @@ TOnlineAsyncOpHandle<FCreateSession> FSessionsEOSGS::CreateSession(FCreateSessio
 
 		if (Result.IsOk())
 		{
-			if (TSharedRef<FSessionEOSGS>* Session = SessionsByName.Find(OpParams.SessionName))
-			{
-				Session->Get().CurrentState = ESessionState::Valid;
+			TSharedRef<FSession>& FoundSession = SessionsByName.FindChecked(OpParams.SessionName);
+			FoundSession->CurrentState = ESessionState::Valid;
 
-				// TODO: Make additional local users join the created session
+			// TODO: Make additional local users join the created session
+
+			Op.SetResult(FCreateSession::Result{ FoundSession });
 
-				Op.SetResult(FCreateSession::Result{ *Session });
-			}
-			else
-			{
-				Op.SetError(Errors::NotFound());
-			}
 		}
 		else
 		{	
 			// If creation failed, we remove the session objects
-			if (TSharedRef<FSessionEOSGS>* Session = SessionsByName.Find(OpParams.SessionName))
-			{
-				SessionsByName.Remove(OpParams.SessionName);
-				SessionsById.Remove(Session->Get().SessionId);
-			}
+			TryRemoveSession(OpParams.SessionName);
 
 			Op.SetError(MoveTemp(Result.GetErrorValue()));
 		}
@@ -397,7 +359,7 @@ void FSessionsEOSGS::SetBucketId(EOS_HSessionModification& SessionModificationHa
 	EOS_EResult ResultCode = EOS_SessionModification_SetBucketId(SessionModificationHandle, &Options);
 	if (ResultCode != EOS_EResult::EOS_Success)
 	{
-		UE_LOG(LogTemp, Warning, TEXT("EOS_SessionModification_SetMaxPlayers failed with result [%s]"), *LexToString(ResultCode));
+		UE_LOG(LogTemp, Warning, TEXT("EOS_SessionModification_SetBucketId failed with result [%s]"), *LexToString(ResultCode));
 	}
 }
 
@@ -465,7 +427,7 @@ void FSessionsEOSGS::WriteCreateSessionModificationHandle(EOS_HSessionModificati
 // 	}
 }
 
-void FSessionsEOSGS::WriteUpdateSessionModificationHandle(EOS_HSessionModification& SessionModificationHandle, FSessionSettings& SessionSettings, const FSessionSettingsUpdate& NewSettings)
+void FSessionsEOSGS::WriteUpdateSessionModificationHandle(EOS_HSessionModification& SessionModificationHandle, const FName& SessionName, const FSessionSettingsUpdate& NewSettings)
 {
 	// TODO: We have the option to call EOS_SessionModification_SetHostAddress in EOS, useful if the session owner changes
 	
@@ -479,89 +441,31 @@ void FSessionsEOSGS::WriteUpdateSessionModificationHandle(EOS_HSessionModificati
 // 	NewSetting.Visibility = ECustomSessionSettingVisibility::ViaOnlineService;
 // 
 // 	AddAttribute(SessionModificationHandle, EOS_SESSIONS_ALLOW_NEW_MEMBERS, NewSetting);
-
-	if (NewSettings.bAllowNewMembers.IsSet())
-	{
-		SessionSettings.bAllowNewMembers = NewSettings.bAllowNewMembers.GetValue();
-	}
-
-	if (NewSettings.bAllowSanctionedPlayers.IsSet())
-	{
-		SessionSettings.bAllowSanctionedPlayers = NewSettings.bAllowSanctionedPlayers.GetValue();
-	}
-
-	if (NewSettings.bAllowUnregisteredPlayers.IsSet())
-	{
-		SessionSettings.bAllowUnregisteredPlayers = NewSettings.bAllowUnregisteredPlayers.GetValue();
-	}
-
-	if (NewSettings.bAntiCheatProtected.IsSet())
-	{
-		SessionSettings.bAntiCheatProtected = NewSettings.bAntiCheatProtected.GetValue();
-	}
-
-	if (NewSettings.bPresenceEnabled.IsSet())
-	{
-		SessionSettings.bPresenceEnabled = NewSettings.bPresenceEnabled.GetValue();
-	}
-
-	if (NewSettings.IsDedicatedServerSession.IsSet())
-	{
-		SessionSettings.IsDedicatedServerSession = NewSettings.IsDedicatedServerSession.GetValue();
-	}
-
-	if (NewSettings.IsLANSession.IsSet())
-	{
-		SessionSettings.IsLANSession = NewSettings.IsLANSession.GetValue();
-	}
+	const TSharedRef<FSession>& FoundSession = SessionsByName.FindChecked(SessionName);
 
 	if (NewSettings.JoinPolicy.IsSet())
 	{
-		SessionSettings.JoinPolicy = NewSettings.JoinPolicy.GetValue();
-
-		SetPermissionLevel(SessionModificationHandle, SessionSettings.JoinPolicy);
+		SetPermissionLevel(SessionModificationHandle, NewSettings.JoinPolicy.GetValue());
 	}
 
-	bool bMaxPlayersUpdated = false;
-
-	if (NewSettings.NumMaxPrivateConnections.IsSet())
+	if (NewSettings.NumMaxPrivateConnections.IsSet() || NewSettings.NumMaxPublicConnections.IsSet())
 	{
-		SessionSettings.NumMaxPrivateConnections = NewSettings.NumMaxPrivateConnections.GetValue();
-		bMaxPlayersUpdated = true;
-	}
+		// We set MaxPlayers to the sum of any new values (if set) and old values (if not set)
+		uint32 DefaultNumMaxPrivateConnections = FoundSession->SessionSettings.NumMaxPrivateConnections;
+		uint32 DefaultNumMaxPublicConnections = FoundSession->SessionSettings.NumMaxPublicConnections;
 
-	if (NewSettings.NumMaxPublicConnections.IsSet())
-	{
-		SessionSettings.NumMaxPublicConnections = NewSettings.NumMaxPublicConnections.GetValue();
-		bMaxPlayersUpdated = true;
+		SetMaxPlayers(SessionModificationHandle, NewSettings.NumMaxPrivateConnections.Get(DefaultNumMaxPrivateConnections) + NewSettings.NumMaxPublicConnections.Get(DefaultNumMaxPublicConnections));
 	}
 
-	if (bMaxPlayersUpdated)
-	{
-		SetMaxPlayers(SessionModificationHandle, SessionSettings.NumMaxPrivateConnections + SessionSettings.NumMaxPublicConnections);
-	}
-
-	if (NewSettings.NumOpenPrivateConnections.IsSet())
-	{
-		SessionSettings.NumOpenPrivateConnections = NewSettings.NumOpenPrivateConnections.GetValue();
-	}
-
-	if (NewSettings.NumOpenPublicConnections.IsSet())
-	{
-		SessionSettings.NumOpenPublicConnections = NewSettings.NumOpenPublicConnections.GetValue();
-	}
-
-	// TODO: We may need some additional logic for schema changes
-	if (NewSettings.SchemaName.IsSet())
-	{
-		SessionSettings.SchemaName = NewSettings.SchemaName.GetValue();
-	}
-
-	if (NewSettings.SessionIdOverride.IsSet())
+	// BucketId has its own set method on the API
+	const FCustomSessionSetting* NewBucketIdSetting = NewSettings.UpdatedCustomSettings.Find(EOS_SESSIONS_BUCKET_ID);
+	if (NewBucketIdSetting)
 	{
-		SessionSettings.SessionIdOverride = NewSettings.SessionIdOverride.GetValue();
+		SetBucketId(SessionModificationHandle, NewBucketIdSetting->Data.Get<FString>());
 	}
 
+	// TODO: This next section will be reactivated as part of the CL to add support for SchemaVariant
+#if 0
 	// Custom Settings
 
 	for (const FName& Key : NewSettings.RemovedCustomSettings)
@@ -578,21 +482,14 @@ void FSessionsEOSGS::WriteUpdateSessionModificationHandle(EOS_HSessionModificati
 
 	SessionSettings.CustomSettings.Append(NewSettings.UpdatedCustomSettings);
 
-	// BucketId has its own set method on the API
-	const FCustomSessionSetting* NewBucketIdSetting = NewSettings.UpdatedCustomSettings.Find(EOS_SESSIONS_BUCKET_ID);
-	if (NewBucketIdSetting)
-	{
-		SetBucketId(SessionModificationHandle, NewBucketIdSetting->Data.Get<FString>());
-	}
-
-	// Registered Users
+	// registered players
 
-	for (const FOnlineAccountIdHandle& Key : NewSettings.RemovedRegisteredUsers)
+	for (const FOnlineAccountIdHandle& Key : NewSettings.RemovedRegisteredPlayers)
 	{
-		SessionSettings.RegisteredUsers.Remove(Key);
+		SessionSettings.RegisteredPlayers.Remove(Key);
 	}
 
-	SessionSettings.RegisteredUsers.Append(NewSettings.UpdatedRegisteredUsers);
+	SessionSettings.RegisteredPlayers.Append(NewSettings.UpdatedRegisteredPlayers);
 
 	// Session Members
 
@@ -641,10 +538,20 @@ void FSessionsEOSGS::WriteUpdateSessionModificationHandle(EOS_HSessionModificati
 			SessionMember->MemberSettings.Append(SessionMemberUpdate.UpdatedMemberSettings);
 		}
 	}
+#endif
 }
 
 TOnlineAsyncOpHandle<FUpdateSession> FSessionsEOSGS::UpdateSession(FUpdateSession::Params&& Params)
 {
+	// LAN Sessions
+	TOnlineResult<FGetSessionByName> Result = GetSessionByName({ Params.SessionName });
+	if (Result.IsOk() && Result.GetOkValue().Session->SessionSettings.bIsLANSession)
+	{
+		return FSessionsLAN::UpdateSession(MoveTemp(Params));
+	}
+
+	// EOSGS Sessions
+
 	TOnlineAsyncOpRef<FUpdateSession> Op = GetOp<FUpdateSession>(MoveTemp(Params));
 	const FUpdateSession::Params& OpParams = Op->GetParams();
 
@@ -652,65 +559,38 @@ TOnlineAsyncOpHandle<FUpdateSession> FSessionsEOSGS::UpdateSession(FUpdateSessio
 	{
 		const FUpdateSession::Params& OpParams = Op.GetParams();
 
-		if (TSharedRef<FSessionEOSGS>* FoundSession = SessionsByName.Find(OpParams.SessionName))
+		FOnlineError StateCheck = CheckUpdateSessionState(OpParams);
+		if (StateCheck != Errors::Success())
 		{
-			if (FoundSession->Get().CurrentState != ESessionState::Valid)
-			{
-				Op.SetError(Errors::InvalidState());
-
-				return MakeFulfilledPromise<TDefaultErrorResult<FUpdateSessionImpl>>().GetFuture();
-			}
-
-			if (!FoundSession->Get().SessionSettings.IsDedicatedServerSession)
-			{
-				IAuthPtr Auth = Services.GetAuthInterface();
-
-				TOnlineResult<FAuthGetLocalOnlineUserByOnlineAccountId> Result = Auth->GetLocalOnlineUserByOnlineAccountId({ OpParams.LocalUserId });
-
-				if (Result.IsOk())
-				{
-					if (Result.GetOkValue().AccountInfo->LoginStatus == ELoginStatus::NotLoggedIn)
-					{
-						UE_LOG(LogTemp, Warning, TEXT("[FSessionsEOSGS::UpdateSession] Could not update session with user [%s] not logged in"), *ToLogString(OpParams.LocalUserId));
+			Op.SetError(MoveTemp(StateCheck));
 
-						Op.SetError(Errors::InvalidUser());
-
-						return MakeFulfilledPromise<TDefaultErrorResult<FUpdateSessionImpl>>().GetFuture();
-					}
-				}
-			}
-
-			EOS_Sessions_UpdateSessionModificationOptions UpdateSessionModificationOptions = {};
-			UpdateSessionModificationOptions.ApiVersion = EOS_SESSIONS_UPDATESESSIONMODIFICATION_API_LATEST;
-			static_assert(EOS_SESSIONS_UPDATESESSIONMODIFICATION_API_LATEST == 1, "EOS_Sessions_UpdateSessionModificationOptions updated, check new fields");
+			return MakeFulfilledPromise<TDefaultErrorResult<FUpdateSessionImpl>>().GetFuture();
+		}
 
-			FTCHARToUTF8 SessionNameUtf8(OpParams.SessionName.ToString());
-			UpdateSessionModificationOptions.SessionName = SessionNameUtf8.Get();
+		EOS_Sessions_UpdateSessionModificationOptions UpdateSessionModificationOptions = {};
+		UpdateSessionModificationOptions.ApiVersion = EOS_SESSIONS_UPDATESESSIONMODIFICATION_API_LATEST;
+		static_assert(EOS_SESSIONS_UPDATESESSIONMODIFICATION_API_LATEST == 1, "EOS_Sessions_UpdateSessionModificationOptions updated, check new fields");
 
-			EOS_HSessionModification SessionModificationHandle = nullptr;
-			EOS_EResult ResultCode = EOS_Sessions_UpdateSessionModification(SessionsHandle, &UpdateSessionModificationOptions, &SessionModificationHandle);
-			if (ResultCode != EOS_EResult::EOS_Success)
-			{
-				UE_LOG(LogTemp, Error, TEXT("[FSessionsEOSGS::UpdateSession] EOS_Sessions_UpdateSessionModification failed with result [%s]"), *LexToString(ResultCode));
+		FTCHARToUTF8 SessionNameUtf8(OpParams.SessionName.ToString());
+		UpdateSessionModificationOptions.SessionName = SessionNameUtf8.Get();
 
-				Op.SetError(FromEOSError(ResultCode));
+		EOS_HSessionModification SessionModificationHandle = nullptr;
+		EOS_EResult ResultCode = EOS_Sessions_UpdateSessionModification(SessionsHandle, &UpdateSessionModificationOptions, &SessionModificationHandle);
+		if (ResultCode != EOS_EResult::EOS_Success)
+		{
+			UE_LOG(LogTemp, Error, TEXT("[FSessionsEOSGS::UpdateSession] EOS_Sessions_UpdateSessionModification failed with result [%s]"), *LexToString(ResultCode));
 
-				return MakeFulfilledPromise<TDefaultErrorResult<FUpdateSessionImpl>>().GetFuture();
-			}
+			Op.SetError(FromEOSError(ResultCode));
 
-			// After creating the session modification handle, we'll add all the updated data to it
-			WriteUpdateSessionModificationHandle(SessionModificationHandle, FoundSession->Get().SessionSettings, OpParams.Mutations);
+			return MakeFulfilledPromise<TDefaultErrorResult<FUpdateSessionImpl>>().GetFuture();
+		}
 
-			FUpdateSessionImpl::Params UpdateSessionImplParams{ SessionModificationHandle };
+		// After creating the session modification handle, we'll add all the updated data to it
+		WriteUpdateSessionModificationHandle(SessionModificationHandle, OpParams.SessionName, OpParams.Mutations);
 
-			return UpdateSessionImpl(MoveTemp(UpdateSessionImplParams));
-		}
-		else
-		{
-			Op.SetError(Errors::NotFound());
+		FUpdateSessionImpl::Params UpdateSessionImplParams{ MakeShared<FSessionModificationHandleEOSGS>(SessionModificationHandle) };
 
-			return MakeFulfilledPromise<TDefaultErrorResult<FUpdateSessionImpl>>().GetFuture();
-		}
+		return UpdateSessionImpl(MoveTemp(UpdateSessionImplParams));
 	})
 	.Then([this](TOnlineAsyncOp<FUpdateSession>& Op, TDefaultErrorResult<FUpdateSessionImpl>&& Result)
 	{
@@ -718,8 +598,11 @@ TOnlineAsyncOpHandle<FUpdateSession> FSessionsEOSGS::UpdateSession(FUpdateSessio
 		{
 			const FUpdateSession::Params& OpParams = Op.GetParams();
 
-			if (TSharedRef<FSessionEOSGS>* Session = SessionsByName.Find(OpParams.SessionName))
+			if (TSharedRef<FSession>* Session = SessionsByName.Find(OpParams.SessionName))
 			{
+				// Now that the API Session update has processed successfully, we'll update our local session with the same data
+				Session->Get().SessionSettings += OpParams.Mutations;
+
 				// We set the result and fire the event
 				Op.SetResult(FUpdateSession::Result{ *Session });
 
@@ -747,7 +630,7 @@ TFuture<TDefaultErrorResult<FUpdateSessionImpl>> FSessionsEOSGS::UpdateSessionIm
 	UpdateSessionOptions.ApiVersion = EOS_SESSIONS_UPDATESESSION_API_LATEST;
 	static_assert(EOS_SESSIONS_UPDATESESSION_API_LATEST == 1, "EOS_Sessions_UpdateSessionOptions updated, check new fields");
 
-	UpdateSessionOptions.SessionModificationHandle = Params.SessionModificationHandle;
+	UpdateSessionOptions.SessionModificationHandle = Params.SessionModificationHandle->ModificationHandle;
 
 	TPromise<TDefaultErrorResult<FUpdateSessionImpl>> Promise;
 	TFuture<TDefaultErrorResult<FUpdateSessionImpl>> Future = Promise.GetFuture();
@@ -761,7 +644,7 @@ TFuture<TDefaultErrorResult<FUpdateSessionImpl>> FSessionsEOSGS::UpdateSessionIm
 			return;
 		}
 
-		if (TSharedRef<FSessionEOSGS>* FoundSession = SessionsByName.Find(FName(Result->SessionName)))
+		if (TSharedRef<FSession>* FoundSession = SessionsByName.Find(FName(Result->SessionName)))
 		{
 			// In session creation calls, we'll need to set the new id for the session
 			if (!FoundSession->Get().SessionId.IsValid())
@@ -778,31 +661,37 @@ TFuture<TDefaultErrorResult<FUpdateSessionImpl>> FSessionsEOSGS::UpdateSessionIm
 
 TOnlineAsyncOpHandle<FLeaveSession> FSessionsEOSGS::LeaveSession(FLeaveSession::Params&& Params)
 {
+	// LAN Sessions
+	TOnlineResult<FGetSessionByName> Result = GetSessionByName({ Params.SessionName });
+	if (Result.IsOk() && Result.GetOkValue().Session->SessionSettings.bIsLANSession)
+	{
+		return FSessionsLAN::LeaveSession(MoveTemp(Params));
+	}
+
+	// EOSGS Sessions
+
 	TOnlineAsyncOpRef<FLeaveSession> Op = GetOp<FLeaveSession>(MoveTemp(Params));
 
 	Op->Then([this](TOnlineAsyncOp<FLeaveSession>& Op, TPromise<const EOS_Sessions_DestroySessionCallbackInfo*>&& Promise)
 	{
 		const FLeaveSession::Params& OpParams = Op.GetParams();
 
-		// We do not check if OpParams.LocalUserId is logged in because it might be a dedicated server destroying the session
-
-		if (TSharedRef<FSessionEOSGS>* FoundSession = SessionsByName.Find(OpParams.SessionName))
+		FOnlineError StateCheck = CheckLeaveSessionState(OpParams);
+		if (StateCheck != Errors::Success())
 		{
-			EOS_Sessions_DestroySessionOptions DestroySessionOptions = {};
-			DestroySessionOptions.ApiVersion = EOS_SESSIONS_DESTROYSESSION_API_LATEST;
-			static_assert(EOS_SESSIONS_DESTROYSESSION_API_LATEST == 1, "EOS_Sessions_DestroySessionOptions updated, check new fields");
-
-			FTCHARToUTF8 SessionNameUtf8(OpParams.SessionName.ToString());
-			DestroySessionOptions.SessionName = SessionNameUtf8.Get();
-
-			EOS_Async(EOS_Sessions_DestroySession, SessionsHandle, DestroySessionOptions, MoveTemp(Promise));
+			Op.SetError(MoveTemp(StateCheck));
+			Promise.EmplaceValue();
 			return;
 		}
-		else
-		{
-			Op.SetError(Errors::NotFound());
-		}
-		Promise.EmplaceValue();
+
+		EOS_Sessions_DestroySessionOptions DestroySessionOptions = {};
+		DestroySessionOptions.ApiVersion = EOS_SESSIONS_DESTROYSESSION_API_LATEST;
+		static_assert(EOS_SESSIONS_DESTROYSESSION_API_LATEST == 1, "EOS_Sessions_DestroySessionOptions updated, check new fields");
+
+		FTCHARToUTF8 SessionNameUtf8(OpParams.SessionName.ToString());
+		DestroySessionOptions.SessionName = SessionNameUtf8.Get();
+
+		EOS_Async(EOS_Sessions_DestroySession, SessionsHandle, DestroySessionOptions, MoveTemp(Promise));
 	})
 	.Then([this](TOnlineAsyncOp<FLeaveSession>& Op, const EOS_Sessions_DestroySessionCallbackInfo* Result)
 	{
@@ -814,11 +703,8 @@ TOnlineAsyncOpHandle<FLeaveSession> FSessionsEOSGS::LeaveSession(FLeaveSession::
 
 		const FLeaveSession::Params& OpParams = Op.GetParams();
 
-		if (TSharedRef<FSessionEOSGS>* FoundSession = SessionsByName.Find(OpParams.SessionName))
+		if (TryRemoveSession(OpParams.SessionName))
 		{
-			SessionsById.Remove((*FoundSession)->SessionId);
-			SessionsByName.Remove(OpParams.SessionName);
-
 			Op.SetResult(FLeaveSession::Result{ });
 
 			FSessionLeft SessionLeftEvent;
@@ -836,58 +722,75 @@ TOnlineAsyncOpHandle<FLeaveSession> FSessionsEOSGS::LeaveSession(FLeaveSession::
 	return Op->GetHandle();
 }
 
-void FSessionsEOSGS::WriteSessionSearchHandle(FSessionSearchHandleEOSGS& SessionSearchHandle, const FFindSessions::Params& Params)
+void FSessionsEOSGS::SetSessionSearchMaxResults(FSessionSearchHandleEOSGS& SessionSearchHandle, uint32 MaxResults)
 {
-	if (Params.MaxResults > 0)
+	EOS_SessionSearch_SetMaxResultsOptions Options = { };
+	Options.ApiVersion = EOS_SESSIONSEARCH_SETMAXSEARCHRESULTS_API_LATEST;
+	static_assert(EOS_SESSIONSEARCH_SETMAXSEARCHRESULTS_API_LATEST == 1, "EOS_SessionSearch_SetMaxResultsOptions updated, check new fields");
+
+	// We truncate the max results parameter up to the value of EOS_SESSIONS_MAX_SEARCH_RESULTS
+	Options.MaxSearchResults = FMath::Clamp(MaxResults, 1, EOS_SESSIONS_MAX_SEARCH_RESULTS);
+
+	EOS_EResult ResultCode = EOS_SessionSearch_SetMaxResults(SessionSearchHandle.SearchHandle, &Options);
+	if (ResultCode != EOS_EResult::EOS_Success)
 	{
-		EOS_SessionSearch_SetMaxResultsOptions Options = { };
-		Options.ApiVersion = EOS_SESSIONSEARCH_SETMAXSEARCHRESULTS_API_LATEST;
-		static_assert(EOS_SESSIONSEARCH_SETMAXSEARCHRESULTS_API_LATEST == 1, "EOS_SessionSearch_SetMaxResultsOptions updated, check new fields");
+		UE_LOG(LogTemp, Warning, TEXT("[FSessionsEOSGS::SetSessionSearchMaxResults] EOS_SessionSearch_SetMaxResults failed with result [%s]"), *LexToString(ResultCode));
+	}
+}
 
-		// We truncate the max results parameter up to the value of EOS_SESSIONS_MAX_SEARCH_RESULTS
-		Options.MaxSearchResults = FMath::Clamp(Params.MaxResults, 1, EOS_SESSIONS_MAX_SEARCH_RESULTS);
+void FSessionsEOSGS::SetSessionSearchParameters(FSessionSearchHandleEOSGS& SessionSearchHandle, TArray<FFindSessionsSearchFilter> Filters)
+{
+	// TODO: Pending SchemaVariant work
+	// Here we'll call EOS_SessionSearch_SetParameter
+}
 
-		EOS_EResult ResultCode = EOS_SessionSearch_SetMaxResults(SessionSearchHandle.SearchHandle, &Options);
-		if (ResultCode != EOS_EResult::EOS_Success)
-		{
-			UE_LOG(LogTemp, Warning, TEXT("[FSessionsEOSGS::WriteSessionSearchHandle] EOS_SessionSearch_SetMaxResults failed with result [%s]"), *LexToString(ResultCode));
-		}
+void FSessionsEOSGS::SetSessionSearchSessionId(FSessionSearchHandleEOSGS& SessionSearchHandle, const FOnlineSessionIdHandle& SessionId)
+{
+	EOS_SessionSearch_SetSessionIdOptions Options = { };
+	Options.ApiVersion = EOS_SESSIONSEARCH_SETSESSIONID_API_LATEST;
+	static_assert(EOS_SESSIONSEARCH_SETSESSIONID_API_LATEST == 1, "EOS_SessionSearch_SetSessionIdOptions updated, check new fields");
+
+	FTCHARToUTF8 SessionIdUtf8(FOnlineSessionIdRegistryEOSGS::Get().ToLogString(SessionId));
+	Options.SessionId = SessionIdUtf8.Get();
+
+	EOS_EResult ResultCode = EOS_SessionSearch_SetSessionId(SessionSearchHandle.SearchHandle, &Options);
+	if (ResultCode != EOS_EResult::EOS_Success)
+	{
+		UE_LOG(LogTemp, Warning, TEXT("[FSessionsEOSGS::SetSessionSearchSessionId] EOS_SessionSearch_SetSessionId failed with result [%s]"), *LexToString(ResultCode));
 	}
+}
+
+void FSessionsEOSGS::SetSessionSearchTargetId(FSessionSearchHandleEOSGS& SessionSearchHandle, const FOnlineAccountIdHandle& TargetUserId)
+{
+	EOS_SessionSearch_SetTargetUserIdOptions Options = { };
+	Options.ApiVersion = EOS_SESSIONSEARCH_SETTARGETUSERID_API_LATEST;
+	static_assert(EOS_SESSIONSEARCH_SETTARGETUSERID_API_LATEST == 1, "EOS_SessionSearch_SetTargetUserIdOptions updated, check new fields");
+	Options.TargetUserId = GetProductUserIdChecked(TargetUserId);
+
+	EOS_EResult ResultCode = EOS_SessionSearch_SetTargetUserId(SessionSearchHandle.SearchHandle, &Options);
+	if (ResultCode != EOS_EResult::EOS_Success)
+	{
+		UE_LOG(LogTemp, Warning, TEXT("[FSessionsEOSGS::SetSessionSearchTargetId] EOS_SessionSearch_SetTargetUserId failed with result [%s]"), *LexToString(ResultCode));
+	}
+}
+
+void FSessionsEOSGS::WriteSessionSearchHandle(FSessionSearchHandleEOSGS& SessionSearchHandle, const FFindSessions::Params& Params)
+{
+	SetSessionSearchMaxResults(SessionSearchHandle, Params.MaxResults);
 
 	if (!Params.Filters.IsEmpty())
 	{
-		// TODO: Pending SchemaVariant work
-		// Here we'll call EOS_SessionSearch_SetParameter
+		SetSessionSearchParameters(SessionSearchHandle, Params.Filters);
 	}
 
 	if (Params.SessionId.IsSet())
 	{
-		EOS_SessionSearch_SetSessionIdOptions Options = { };
-		Options.ApiVersion = EOS_SESSIONSEARCH_SETSESSIONID_API_LATEST;
-		static_assert(EOS_SESSIONSEARCH_SETSESSIONID_API_LATEST == 1, "EOS_SessionSearch_SetSessionIdOptions updated, check new fields");
-
-		FTCHARToUTF8 SessionIdUtf8(FOnlineSessionIdRegistryEOSGS::Get().ToLogString(*Params.SessionId));
-		Options.SessionId = SessionIdUtf8.Get();
-
-		EOS_EResult ResultCode = EOS_SessionSearch_SetSessionId(SessionSearchHandle.SearchHandle, &Options);
-		if (ResultCode != EOS_EResult::EOS_Success)
-		{
-			UE_LOG(LogTemp, Warning, TEXT("[FSessionsEOSGS::WriteSessionSearchHandle] EOS_SessionSearch_SetSessionId failed with result [%s]"), *LexToString(ResultCode));
-		}
+		SetSessionSearchSessionId(SessionSearchHandle, Params.SessionId.GetValue());
 	}
 
 	if (Params.TargetUser.IsSet())
 	{
-		EOS_SessionSearch_SetTargetUserIdOptions Options = { };
-		Options.ApiVersion = EOS_SESSIONSEARCH_SETTARGETUSERID_API_LATEST;
-		static_assert(EOS_SESSIONSEARCH_SETTARGETUSERID_API_LATEST == 1, "EOS_SessionSearch_SetTargetUserIdOptions updated, check new fields");
-		Options.TargetUserId = GetProductUserIdChecked(*Params.TargetUser);
-
-		EOS_EResult ResultCode = EOS_SessionSearch_SetTargetUserId(SessionSearchHandle.SearchHandle, &Options);
-		if (ResultCode != EOS_EResult::EOS_Success)
-		{
-			UE_LOG(LogTemp, Warning, TEXT("[FSessionsEOSGS::WriteSessionSearchHandle] EOS_SessionSearch_SetTargetUserId failed with result [%s]"), *LexToString(ResultCode));
-		}
+		SetSessionSearchTargetId(SessionSearchHandle, Params.TargetUser.GetValue());
 	}
 }
 
@@ -898,78 +801,67 @@ FOnlineSessionIdHandle FSessionsEOSGS::CreateSessionId(const FString& SessionId)
 
 TOnlineAsyncOpHandle<FFindSessions> FSessionsEOSGS::FindSessions(FFindSessions::Params&& Params)
 {
-	TOnlineAsyncOpRef<FFindSessions> Op = GetOp<FFindSessions>(MoveTemp(Params));
-	const FFindSessions::Params& OpParams = Op->GetParams();
-
-	// We check if the user is logged in
-	IAuthPtr Auth = Services.GetAuthInterface();
-	TOnlineResult<FAuthGetLocalOnlineUserByOnlineAccountId> Result = Auth->GetLocalOnlineUserByOnlineAccountId({ OpParams.LocalUserId });
-
-	if (Result.IsOk())
+	// LAN Sessions
+	if (Params.bFindLANSessions)
 	{
-		if (Result.GetOkValue().AccountInfo->LoginStatus == ELoginStatus::NotLoggedIn)
-		{
-			UE_LOG(LogTemp, Warning, TEXT("[FSessionsEOSGS::FindSessions] Could not find sessions with user [%s] not logged in"), *ToLogString(OpParams.LocalUserId));
-
-			Op->SetError(Errors::InvalidUser());
-			return Op->GetHandle();
-		}
+		return FSessionsLAN::FindSessions(MoveTemp(Params));
 	}
 
-	// We check if MaxResults has a valid value
-	if (OpParams.MaxResults > 0 && OpParams.MaxResults < EOS_SESSIONS_MAX_SEARCH_RESULTS)
-	{
-		UE_LOG(LogTemp, Error, TEXT("[FSessionsEOSGS::FindSessions] FFindSessions::Params::MaxResults had invalid value [%d]. Value must be between [0-%d]"), Params.MaxResults, EOS_SESSIONS_MAX_SEARCH_RESULTS);
+	// EOSGS Sessions
 
-		Op->SetError(Errors::InvalidParams());
+	TOnlineAsyncOpRef<FFindSessions> Op = GetOp<FFindSessions>(MoveTemp(Params));
+	const FFindSessions::Params& OpParams = Op->GetParams();
+
+	FOnlineError ParamsCheck = CheckFindSessionsParams(OpParams);
+	if (ParamsCheck != Errors::Success())
+	{
+		Op->SetError(MoveTemp(ParamsCheck));
 		return Op->GetHandle();
 	}
 
 	// We start the find operation
 	Op->Then([this](TOnlineAsyncOp<FFindSessions>& Op, TPromise<const EOS_SessionSearch_FindCallbackInfo*>&& Promise)
 	{
-		if (!CurrentSessionSearch.IsValid())
-		{
-			CurrentSessionSearch = MakeShared<FFindSessions::Result>();
-			CurrentSessionSearchAsyncOpHandle = Op.AsShared();
+		const FFindSessions::Params& OpParams = Op.GetParams();
 
-			const FFindSessions::Params& OpParams = Op.GetParams();
+		FOnlineError StateCheck = CheckFindSessionsState(OpParams);
+		if (StateCheck != Errors::Success())
+		{
+			Op.SetError(MoveTemp(StateCheck));
+			Promise.EmplaceValue();
+			return;
+		}
 
-			EOS_Sessions_CreateSessionSearchOptions CreateSessionSearchOptions = { };
-			CreateSessionSearchOptions.ApiVersion = EOS_SESSIONS_CREATESESSIONSEARCH_API_LATEST;
-			static_assert(EOS_SESSIONS_CREATESESSIONSEARCH_API_LATEST == 1, "EOS_Sessions_CreateSessionSearchOptions updated, check new fields");
-			CreateSessionSearchOptions.MaxSearchResults = FMath::Clamp(OpParams.MaxResults, 0, EOS_SESSIONS_MAX_SEARCH_RESULTS);
+		CurrentSessionSearch = MakeShared<FFindSessions::Result>();
+		CurrentSessionSearchHandle = Op.AsShared();
 
-			EOS_HSessionSearch SearchHandle = nullptr;
-			EOS_EResult ResultCode = EOS_Sessions_CreateSessionSearch(SessionsHandle, &CreateSessionSearchOptions, &SearchHandle);
-			if (ResultCode != EOS_EResult::EOS_Success)
-			{
-				UE_LOG(LogTemp, Error, TEXT("[FSessionsEOSGS::FindSessions] EOS_Sessions_CreateSessionSearch failed with result [%s]"), *LexToString(ResultCode));
+		EOS_Sessions_CreateSessionSearchOptions CreateSessionSearchOptions = { };
+		CreateSessionSearchOptions.ApiVersion = EOS_SESSIONS_CREATESESSIONSEARCH_API_LATEST;
+		static_assert(EOS_SESSIONS_CREATESESSIONSEARCH_API_LATEST == 1, "EOS_Sessions_CreateSessionSearchOptions updated, check new fields");
+		CreateSessionSearchOptions.MaxSearchResults = FMath::Clamp(OpParams.MaxResults, 0, EOS_SESSIONS_MAX_SEARCH_RESULTS);
 
-				Op.SetError(FromEOSError(ResultCode));
+		EOS_HSessionSearch SearchHandle = nullptr;
+		EOS_EResult ResultCode = EOS_Sessions_CreateSessionSearch(SessionsHandle, &CreateSessionSearchOptions, &SearchHandle);
+		if (ResultCode != EOS_EResult::EOS_Success)
+		{
+			UE_LOG(LogTemp, Error, TEXT("[FSessionsEOSGS::FindSessions] EOS_Sessions_CreateSessionSearch failed with result [%s]"), *LexToString(ResultCode));
 
-				Promise.EmplaceValue();
-				return;
-			}
+			Op.SetError(FromEOSError(ResultCode));
+			Promise.EmplaceValue();
+			return;
+		}
 
-			CurrentSessionSearchHandleEOSGS = MakeShared<FSessionSearchHandleEOSGS>(SearchHandle);
+		CurrentSessionSearchHandleEOSGS = MakeShared<FSessionSearchHandleEOSGS>(SearchHandle);
 
-			// We write the search attributes
-			WriteSessionSearchHandle(*CurrentSessionSearchHandleEOSGS, OpParams);
+		// We write the search attributes
+		WriteSessionSearchHandle(*CurrentSessionSearchHandleEOSGS, OpParams);
 
-			EOS_SessionSearch_FindOptions FindOptions = { };
-			FindOptions.ApiVersion = EOS_SESSIONSEARCH_FIND_API_LATEST;
-			static_assert(EOS_SESSIONSEARCH_FIND_API_LATEST == 2, "EOS_SessionSearch_FindOptions updated, check new fields");
-			FindOptions.LocalUserId = GetProductUserIdChecked(OpParams.LocalUserId);
+		EOS_SessionSearch_FindOptions FindOptions = { };
+		FindOptions.ApiVersion = EOS_SESSIONSEARCH_FIND_API_LATEST;
+		static_assert(EOS_SESSIONSEARCH_FIND_API_LATEST == 2, "EOS_SessionSearch_FindOptions updated, check new fields");
+		FindOptions.LocalUserId = GetProductUserIdChecked(OpParams.LocalUserId);
 
-			EOS_Async(EOS_SessionSearch_Find, CurrentSessionSearchHandleEOSGS->SearchHandle, FindOptions, MoveTemp(Promise));
-			return;
-		}
-		else
-		{
-			Op.SetError(Errors::AlreadyPending());
-		}
-		Promise.EmplaceValue();
+		EOS_Async(EOS_SessionSearch_Find, CurrentSessionSearchHandleEOSGS->SearchHandle, FindOptions, MoveTemp(Promise));
 	})
 	.Then([this](TOnlineAsyncOp<FFindSessions>& Op, const EOS_SessionSearch_FindCallbackInfo*&& FindCallbackInfoResult)
 	{
@@ -998,9 +890,7 @@ TOnlineAsyncOpHandle<FFindSessions> FSessionsEOSGS::FindSessions(FFindSessions::
 			if (CopySearchResultByIndexResult == EOS_EResult::EOS_Success)
 			{
 				// All session attributes get extracted with the SessionDetailsHandle
-				FSessionEOSGS NewSession = FSessionEOSGS(SessionDetailsHandle);
-
-				TSharedRef<FSessionEOSGS> NewSessionEOSGSRef = MakeShared<FSessionEOSGS>(NewSession);
+				const TSharedRef<FSessionEOSGS>& NewSessionEOSGSRef = MakeShared<FSessionEOSGS>(SessionDetailsHandle);
 
 				CurrentSessionSearch->FoundSessions.Add(NewSessionEOSGSRef);
 			}
@@ -1022,20 +912,17 @@ TOnlineAsyncOpHandle<FFindSessions> FSessionsEOSGS::FindSessions(FFindSessions::
 
 TOnlineAsyncOpHandle<FJoinSession> FSessionsEOSGS::JoinSession(FJoinSession::Params&& Params)
 {
-	TOnlineAsyncOpRef<FJoinSession> Op = GetOp<FJoinSession>(MoveTemp(Params));
-	const FJoinSession::Params& OpParams = Op->GetParams();
-
-	// TODO: Call CheckJoinSessionParams after it's submitted, add the state check to it
-	// 
-	// We check that the passed session is valid
-	if (Params.Session->CurrentState != ESessionState::Valid)
+	// LAN Sessions
+	if (Params.Session->SessionSettings.bIsLANSession)
 	{
-		UE_LOG(LogTemp, Warning, TEXT("[FSessionsEOSGS::FJoinSession] Could not join session with invalid session state [%s]"), *LexToString(Params.Session->CurrentState));
-
-		Op->SetError(Errors::InvalidParams());
-		return Op->GetHandle();
+		return FSessionsLAN::JoinSession(MoveTemp(Params));
 	}
 
+	// EOSGS Sessions
+
+	TOnlineAsyncOpRef<FJoinSession> Op = GetOp<FJoinSession>(MoveTemp(Params));
+	const FJoinSession::Params& OpParams = Op->GetParams();
+
 	// We check that the passed session has a valid details handle
 	const FSessionEOSGS& SessionEOSGS = FSessionEOSGS::Cast(*Params.Session);
 	if (!SessionEOSGS.SessionDetailsHandle.IsValid())
@@ -1046,23 +933,32 @@ TOnlineAsyncOpHandle<FJoinSession> FSessionsEOSGS::JoinSession(FJoinSession::Par
 		return Op->GetHandle();
 	}
 
-	// TODO: Routine JoinSession checks
+	FOnlineError ParamsCheck = CheckJoinSessionParams(OpParams);
+	if (ParamsCheck != Errors::Success())
+	{
+		Op->SetError(MoveTemp(ParamsCheck));
+		return Op->GetHandle();
+	}
 
 	// We start the join operation
 	Op->Then([this](TOnlineAsyncOp<FJoinSession>& Op, TPromise<const EOS_Sessions_JoinSessionCallbackInfo*>&& Promise)
 	{
 		const FJoinSession::Params& OpParams = Op.GetParams();
 
-		// TODO: Call CheckJoinSessionState after it's submitted
-
-		// We make a copy of the Session object
-		FSessionEOSGS NewSession = FSessionEOSGS(*OpParams.Session);
-		NewSession.CurrentState = ESessionState::Joining;
+		FOnlineError StateCheck = CheckJoinSessionState(OpParams);
+		if (StateCheck != Errors::Success())
+		{
+			Op.SetError(MoveTemp(StateCheck));
+			Promise.EmplaceValue();
+			return;
+		}
 
-		TSharedRef<FSessionEOSGS> NewSessionEOSGSRef = MakeShared<FSessionEOSGS>(NewSession);
+		// We make a copy of the Session object		
+		TSharedRef<FSession> NewSessionRef = MakeShared<FSession>(*OpParams.Session);
+		NewSessionRef->CurrentState = ESessionState::Joining;
 
-		SessionsByName.Add(OpParams.SessionName, NewSessionEOSGSRef);
-		SessionsById.Add(NewSession.SessionId, NewSessionEOSGSRef);
+		SessionsByName.Add(OpParams.SessionName, NewSessionRef);
+		SessionsById.Add(NewSessionRef->SessionId, NewSessionRef);
 
 		// We start setup for the API call
 		EOS_Sessions_JoinSessionOptions JoinSessionOptions = { };
@@ -1087,18 +983,13 @@ TOnlineAsyncOpHandle<FJoinSession> FSessionsEOSGS::JoinSession(FJoinSession::Par
 
 		if (Result->ResultCode != EOS_EResult::EOS_Success)
 		{
-			// If join failed, we remove the session objects
-			if (TSharedRef<FSessionEOSGS>* Session = SessionsByName.Find(OpParams.SessionName))
-			{
-				SessionsByName.Remove(OpParams.SessionName);
-				SessionsById.Remove(Session->Get().SessionId);
-			}
+			TryRemoveSession(OpParams.SessionName);
 
 			Op.SetError(FromEOSError(Result->ResultCode));
 			return;
 		}
 		
-		if (TSharedRef<FSessionEOSGS>* FoundSession = SessionsByName.Find(OpParams.SessionName))
+		if (TSharedRef<FSession>* FoundSession = SessionsByName.Find(OpParams.SessionName))
 		{
 			FoundSession->Get().CurrentState = ESessionState::Valid;
 
@@ -1186,12 +1077,17 @@ TOnlineAsyncOpHandle<FSendSessionInvite> FSessionsEOSGS::SendSessionInvite(FSend
 {
 	TOnlineAsyncOpRef<FSendSessionInvite> Op = GetOp<FSendSessionInvite>(MoveTemp(Params));
 
-	// TODO: Param checks
 	Op->Then([this](TOnlineAsyncOp<FSendSessionInvite>& Op, TPromise<const EOS_Sessions_SendInviteCallbackInfo*>&& Promise)
 	{
 		const FSendSessionInvite::Params& OpParams = Op.GetParams();
 
-		// TODO: State checks
+		FOnlineError StateCheck = CheckSendSessionInviteState(OpParams);
+		if (StateCheck != Errors::Success())
+		{
+			Op.SetError(MoveTemp(StateCheck));
+			Promise.EmplaceValue();
+			return;
+		}
 
 		EOS_Sessions_SendInviteOptions SendInviteOptions = { };
 		SendInviteOptions.ApiVersion = EOS_SESSIONS_SENDINVITE_API_LATEST;
@@ -1204,7 +1100,7 @@ TOnlineAsyncOpHandle<FSendSessionInvite> FSessionsEOSGS::SendSessionInvite(FSend
 
 		SendInviteOptions.TargetUserId = GetProductUserIdChecked(OpParams.TargetUsers[0]); // TODO: Multiple user invitations using WhenAll like JoinLobbyImpl
 
-		EOS_Async<EOS_Sessions_SendInviteCallbackInfo>(EOS_Sessions_SendInvite, SessionsHandle, SendInviteOptions, MoveTemp(Promise));
+		EOS_Async(EOS_Sessions_SendInvite, SessionsHandle, SendInviteOptions, MoveTemp(Promise));
 		return;
 	})
 	.Then([this](TOnlineAsyncOp<FSendSessionInvite>& Op, const EOS_Sessions_SendInviteCallbackInfo* Result)
@@ -1226,12 +1122,17 @@ TOnlineAsyncOpHandle<FRejectSessionInvite> FSessionsEOSGS::RejectSessionInvite(F
 {
 	TOnlineAsyncOpRef<FRejectSessionInvite> Op = GetOp<FRejectSessionInvite>(MoveTemp(Params));
 
-	// TODO: Param checks
 	Op->Then([this](TOnlineAsyncOp<FRejectSessionInvite>& Op, TPromise<const EOS_Sessions_RejectInviteCallbackInfo*>&& Promise)
 	{
 		const FRejectSessionInvite::Params& OpParams = Op.GetParams();
 
-		// TODO: State checks
+		FOnlineError StateCheck = CheckRejectSessionInviteState(OpParams);
+		if (StateCheck != Errors::Success())
+		{
+			Op.SetError(MoveTemp(StateCheck));
+			Promise.EmplaceValue();
+			return;
+		}
 
 		EOS_Sessions_RejectInviteOptions RejectInviteOptions = { };
 		RejectInviteOptions.ApiVersion = EOS_SESSIONS_REJECTINVITE_API_LATEST;
@@ -1242,7 +1143,7 @@ TOnlineAsyncOpHandle<FRejectSessionInvite> FSessionsEOSGS::RejectSessionInvite(F
 		const FTCHARToUTF8 InviteIdUtf8(OpParams.SessionInvite->InviteId);
 		RejectInviteOptions.InviteId = InviteIdUtf8.Get();
 
-		EOS_Async<EOS_Sessions_RejectInviteCallbackInfo>(EOS_Sessions_RejectInvite, SessionsHandle, RejectInviteOptions, MoveTemp(Promise));
+		EOS_Async(EOS_Sessions_RejectInvite, SessionsHandle, RejectInviteOptions, MoveTemp(Promise));
 		return;
 	})
 	.Then([this](TOnlineAsyncOp<FRejectSessionInvite>& Op, const EOS_Sessions_RejectInviteCallbackInfo* Result)
@@ -1260,4 +1161,186 @@ TOnlineAsyncOpHandle<FRejectSessionInvite> FSessionsEOSGS::RejectSessionInvite(F
 	return Op->GetHandle();
 }
 
+TOnlineAsyncOpHandle<FAddSessionMembers> FSessionsEOSGS::AddSessionMembers(FAddSessionMembers::Params&& Params)
+{
+	// LAN Sessions
+	TOnlineResult<FGetSessionByName> Result = GetSessionByName({ Params.SessionName });
+	if (Result.IsOk() && Result.GetOkValue().Session->SessionSettings.bIsLANSession)
+	{
+		return FSessionsLAN::AddSessionMembers(MoveTemp(Params));
+	}
+
+	// EOSGS Sessions
+
+	TOnlineAsyncOpRef<FAddSessionMembers> Op = GetOp<FAddSessionMembers>(MoveTemp(Params));
+
+	Op->Then([this](TOnlineAsyncOp<FAddSessionMembers>& Op, TPromise<const EOS_Sessions_RegisterPlayersCallbackInfo*>&& Promise)
+	{
+		const FAddSessionMembers::Params& OpParams = Op.GetParams();
+
+		FOnlineError StateCheck = CheckAddSessionMembersState(OpParams);
+		if (StateCheck != Errors::Success())
+		{
+			Op.SetError(MoveTemp(StateCheck));
+			Promise.EmplaceValue();
+			return;
+		}
+
+		EOS_Sessions_RegisterPlayersOptions Options = { };
+		Options.ApiVersion = EOS_SESSIONS_REGISTERPLAYERS_API_LATEST;
+		static_assert(EOS_SESSIONS_REGISTERPLAYERS_API_LATEST == 2, "EOS_Sessions_RegisterPlayersOptions updated, check new fields");
+
+		Options.PlayersToRegisterCount = OpParams.NewSessionMembers.Num();
+
+		TArray<EOS_ProductUserId> PlayersToRegister;
+		TArray<FOnlineAccountIdHandle> SessionMemberIds;
+		OpParams.NewSessionMembers.GenerateKeyArray(SessionMemberIds);		
+		for (FOnlineAccountIdHandle IdHandle : SessionMemberIds)
+		{
+			// TODO: Do this with ResolveAccountIds instead
+			PlayersToRegister.Add(GetProductUserIdChecked(IdHandle));
+		}
+
+		Options.PlayersToRegister = PlayersToRegister.GetData();
+
+		FTCHARToUTF8 SessionNameUtf8(OpParams.SessionName.ToString());
+		Options.SessionName = SessionNameUtf8.Get();
+
+		EOS_Async(EOS_Sessions_RegisterPlayers, SessionsHandle, Options, MoveTemp(Promise));
+	})
+	.Then([this](TOnlineAsyncOp<FAddSessionMembers>& Op, const EOS_Sessions_RegisterPlayersCallbackInfo* Result)
+	{
+		if (Result->ResultCode != EOS_EResult::EOS_Success)
+		{
+			Op.SetError(FromEOSError(Result->ResultCode));
+			return;
+		}
+
+		TOnlineResult<FAddSessionMembers> ImplResult = AddSessionMembersImpl(Op.GetParams());
+		if (ImplResult.IsOk())
+		{
+			Op.SetResult(MoveTemp(ImplResult.GetOkValue()));
+		}
+		else
+		{
+			Op.SetError(MoveTemp(ImplResult.GetErrorValue()));
+		}
+	})
+	.Enqueue(GetSerialQueue());
+
+	return Op->GetHandle();
+}
+
+TOnlineAsyncOpHandle<FRemoveSessionMembers> FSessionsEOSGS::RemoveSessionMembers(FRemoveSessionMembers::Params&& Params)
+{
+	// LAN Sessions
+	TOnlineResult<FGetSessionByName> Result = GetSessionByName({ Params.SessionName });
+	if (Result.IsOk() && Result.GetOkValue().Session->SessionSettings.bIsLANSession)
+	{
+		return FSessionsLAN::RemoveSessionMembers(MoveTemp(Params));
+	}
+
+	// EOSGS Sessions
+
+	TOnlineAsyncOpRef<FRemoveSessionMembers> Op = GetOp<FRemoveSessionMembers>(MoveTemp(Params));
+
+	Op->Then([this](TOnlineAsyncOp<FRemoveSessionMembers>& Op, TPromise<const EOS_Sessions_UnregisterPlayersCallbackInfo*>&& Promise)
+	{
+		const FRemoveSessionMembers::Params& OpParams = Op.GetParams();
+
+		FOnlineError StateCheck = CheckRemoveSessionMembersState(OpParams);
+		if (StateCheck != Errors::Success())
+		{
+			Op.SetError(MoveTemp(StateCheck));
+			Promise.EmplaceValue();
+			return;
+		}
+
+		EOS_Sessions_UnregisterPlayersOptions Options = { };
+		Options.ApiVersion = EOS_SESSIONS_UNREGISTERPLAYERS_API_LATEST;
+		static_assert(EOS_SESSIONS_UNREGISTERPLAYERS_API_LATEST == 2, "EOS_Sessions_UnregisterPlayersOptions updated, check new fields");
+
+		Options.PlayersToUnregisterCount = OpParams.SessionMemberIds.Num();
+
+		TArray<EOS_ProductUserId> PlayersToUnregister;
+		for (FOnlineAccountIdHandle IdHandle : OpParams.SessionMemberIds)
+		{
+			// TODO: Do this with ResolveAccountIds instead
+			PlayersToUnregister.Add(GetProductUserIdChecked(IdHandle));
+		}
+
+		Options.PlayersToUnregister = PlayersToUnregister.GetData();
+
+		FTCHARToUTF8 SessionNameUtf8(OpParams.SessionName.ToString());
+		Options.SessionName = SessionNameUtf8.Get();
+
+		EOS_Async(EOS_Sessions_UnregisterPlayers, SessionsHandle, Options, MoveTemp(Promise));
+	})
+	.Then([this](TOnlineAsyncOp<FRemoveSessionMembers>& Op, const EOS_Sessions_UnregisterPlayersCallbackInfo* Result)
+	{
+		if (Result->ResultCode != EOS_EResult::EOS_Success)
+		{
+			Op.SetError(FromEOSError(Result->ResultCode));
+			return;
+		}
+
+		TOnlineResult<FRemoveSessionMembers> ImplResult = RemoveSessionMembersImpl(Op.GetParams());
+		if (ImplResult.IsOk())
+		{
+			Op.SetResult(MoveTemp(ImplResult.GetOkValue()));
+		}
+		else
+		{
+			Op.SetError(MoveTemp(ImplResult.GetErrorValue()));
+		}
+	})
+	.Enqueue(GetSerialQueue());
+
+	return Op->GetHandle();
+}
+
+bool FSessionsEOSGS::TryRemoveSession(const FName& SessionName)
+{
+	if (TSharedRef<FSession>* FoundSession = SessionsByName.Find(SessionName))
+	{
+		SessionsById.Remove((*FoundSession)->SessionId);
+		SessionsByName.Remove(SessionName);
+
+		return true;
+	}
+
+	return false;
+}
+
+/** FSessionsLAN */
+
+void FSessionsEOSGS::AppendSessionToPacket(FNboSerializeToBuffer& Packet, const FSessionLAN& Session)
+{
+	using namespace NboSerializerEOSGSSvc;
+
+	// We won't save CurrentState in the packet as all advertised sessions will be Valid
+	SerializeToBuffer(Packet, Session.OwnerUserId);
+	SerializeToBuffer(Packet, Session.SessionId);
+	Packet << *Session.OwnerInternetAddr;
+
+	// TODO: Write session settings to packet, after SchemaVariant work
+}
+
+void FSessionsEOSGS::ReadSessionFromPacket(FNboSerializeFromBuffer& Packet, FSessionLAN& Session)
+{
+	using namespace NboSerializerEOSGSSvc;
+
+	SerializeFromBuffer(Packet, Session.OwnerUserId);
+	SerializeFromBuffer(Packet, Session.SessionId);
+	Packet >> *Session.OwnerInternetAddr;
+
+	// We'll set the connect address for the remote session as a custom parameter, so it can be read in OnlineServices' GetResolvedConnectString
+	FCustomSessionSetting ConnectString;
+	ConnectString.Data.Set<FString>(Session.OwnerInternetAddr->ToString(true));
+	ConnectString.Visibility = ECustomSessionSettingVisibility::ViaOnlineService;
+	Session.SessionSettings.CustomSettings.Add(CONNECT_STRING_TAG, ConnectString);
+
+	// TODO:: Read session settings from packet, after SchemaVariant work
+}
+
 /* UE::Online */ }
\ No newline at end of file
diff --git a/Engine/Plugins/Online/Experimental/OnlineServicesEOSGS/Source/Public/Online/NboSerializerEOSGSSvc.h b/Engine/Plugins/Online/Experimental/OnlineServicesEOSGS/Source/Public/Online/NboSerializerEOSGSSvc.h
new file mode 100644
--- /dev/null
+++ b/Engine/Plugins/Online/Experimental/OnlineServicesEOSGS/Source/Public/Online/NboSerializerEOSGSSvc.h
@@ -0,0 +1,57 @@
+// Copyright Epic Games, Inc. All Rights Reserved.
+
+#pragma once
+
+#include "Online/NboSerializer.h"
+
+#include "Online/AuthEOSGS.h"
+#include "Online/LobbiesEOSGS.h"
+#include "Online/SessionsEOSGS.h"
+#include "Online/CoreOnline.h"
+
+/**
+ * Serializes data in network byte order form into a buffer
+ */
+namespace UE::Online {
+
+namespace NboSerializerEOSGSSvc {
+
+/** NboSerializeToBuffer methods */
+
+inline void SerializeToBuffer(FNboSerializeToBuffer& Ar, const FOnlineAccountIdHandle& UniqueId)
+{
+	TArray<uint8> Data = FOnlineAccountIdRegistryEOSGS::GetRegistered().ToReplicationData(UniqueId);
+	Ar << Data.Num();
+	Ar.WriteBinary(Data.GetData(), Data.Num());
+}
+
+inline void SerializeToBuffer(FNboSerializeToBuffer& Ar, const FOnlineSessionIdHandle& SessionId)
+{
+	TArray<uint8> Data = FOnlineSessionIdRegistryEOSGS::Get().ToReplicationData(SessionId);
+	Ar << Data.Num();
+	Ar.WriteBinary(Data.GetData(), Data.Num());
+}
+
+/** NboSerializeFromBuffer methods */
+
+inline void SerializeFromBuffer(FNboSerializeFromBuffer& Ar, FOnlineAccountIdHandle& UniqueId)
+{
+	TArray<uint8> Data;
+	int32 Size;
+	Ar >> Size;
+	Ar.ReadBinaryArray(Data, Size);
+	UniqueId = FOnlineAccountIdRegistryEOSGS::GetRegistered().FromReplicationData(Data);
+}
+
+inline void SerializeFromBuffer(FNboSerializeFromBuffer& Ar, FOnlineSessionIdHandle& SessionId)
+{
+	TArray<uint8> Data;
+	int32 Size;
+	Ar >> Size;
+	Ar.ReadBinaryArray(Data, Size);
+	SessionId = FOnlineSessionIdRegistryEOSGS::Get().FromReplicationData(Data);
+}
+
+/* NboSerializerNullSvc */ }
+
+/* UE::Online */ }
\ No newline at end of file
diff --git a/Engine/Plugins/Online/Experimental/OnlineServicesEOSGS/Source/Public/Online/SessionsEOSGS.h b/Engine/Plugins/Online/Experimental/OnlineServicesEOSGS/Source/Public/Online/SessionsEOSGS.h
--- a/Engine/Plugins/Online/Experimental/OnlineServicesEOSGS/Source/Public/Online/SessionsEOSGS.h
+++ b/Engine/Plugins/Online/Experimental/OnlineServicesEOSGS/Source/Public/Online/SessionsEOSGS.h
@@ -2,10 +2,9 @@
 
 #pragma once
 
-#include "CoreMinimal.h"
-#include "Online/SessionsCommon.h"
+#include "Online/SessionsLAN.h"
+
 #include "Online/OnlineServicesEOSGSTypes.h"
-#include "IPAddress.h"
 
 #if defined(EOS_PLATFORM_BASE_FILE_NAME)
 #include EOS_PLATFORM_BASE_FILE_NAME
@@ -25,6 +24,21 @@ public:
 
 static FName EOS_SESSIONS_BUCKET_ID = TEXT("EOS_SESSIONS_BUCKET_ID");
 
+struct FSessionModificationHandleEOSGS : FNoncopyable
+{
+	EOS_HSessionModification ModificationHandle;
+
+	FSessionModificationHandleEOSGS(EOS_HSessionModification InModificationHandle)
+		: ModificationHandle(InModificationHandle)
+	{
+	}
+
+	~FSessionModificationHandleEOSGS()
+	{
+		EOS_SessionModification_Release(ModificationHandle);
+	}
+};
+
 struct FSessionSearchHandleEOSGS : FNoncopyable
 {
 	EOS_HSessionSearch SearchHandle;
@@ -55,11 +69,11 @@ struct FSessionDetailsHandleEOSGS : FNoncopyable
 	}
 };
 
-class FSessionEOSGS : public FSession
+class FSessionEOSGS : public FSessionLAN
 {
 public:
-	FSessionEOSGS();
-	FSessionEOSGS(const FSession& InSession);
+	FSessionEOSGS() = default;
+	FSessionEOSGS(const FSessionEOSGS& InSession) = default;
 	FSessionEOSGS(const EOS_HSessionDetails& SessionDetailsHandle);
 
 	static const FSessionEOSGS& Cast(const FSession& InSession);
@@ -76,7 +90,7 @@ struct FUpdateSessionImpl
 	struct Params
 	{
 		/** Handle for the session modification operation */
-		EOS_HSessionModification SessionModificationHandle;
+		TSharedRef<FSessionModificationHandleEOSGS> SessionModificationHandle;
 	};
 
 	struct Result
@@ -84,12 +98,12 @@ struct FUpdateSessionImpl
 	};
 };
 
-class ONLINESERVICESEOSGS_API FSessionsEOSGS : public FSessionsCommon
+class ONLINESERVICESEOSGS_API FSessionsEOSGS : public FSessionsLAN
 {
 public:
 	friend class FSessionEOSGS;
 
-	using Super = FSessionsCommon;
+	using Super = FSessionsLAN;
 
 	FSessionsEOSGS(FOnlineServicesEOSGS& InOwningSubsystem);
 	virtual ~FSessionsEOSGS() = default;
@@ -106,6 +120,8 @@ public:
 	virtual TOnlineAsyncOpHandle<FJoinSession> JoinSession(FJoinSession::Params&& Params) override;
 	virtual TOnlineAsyncOpHandle<FSendSessionInvite> SendSessionInvite(FSendSessionInvite::Params&& Params) override;
 	virtual TOnlineAsyncOpHandle<FRejectSessionInvite> RejectSessionInvite(FRejectSessionInvite::Params&& Params) override;
+	virtual TOnlineAsyncOpHandle<FAddSessionMembers> AddSessionMembers(FAddSessionMembers::Params&& Params) override;
+	virtual TOnlineAsyncOpHandle<FRemoveSessionMembers> RemoveSessionMembers(FRemoveSessionMembers::Params&& Params) override;
 
 protected:
 	void RegisterEventHandlers();
@@ -118,14 +134,19 @@ protected:
 	void SetBucketId(EOS_HSessionModification& SessionModHandle, const FString& NewBucketId);
 	void SetMaxPlayers(EOS_HSessionModification& SessionModHandle, const uint32& NewMaxPlayers);
 
+	void SetSessionSearchMaxResults(FSessionSearchHandleEOSGS& SessionSearchHandle, uint32 MaxResults);
+	void SetSessionSearchParameters(FSessionSearchHandleEOSGS& SessionSearchHandle, TArray<FFindSessionsSearchFilter> Filters);
+	void SetSessionSearchSessionId(FSessionSearchHandleEOSGS& SessionSearchHandle, const FOnlineSessionIdHandle& SessionId);
+	void SetSessionSearchTargetId(FSessionSearchHandleEOSGS& SessionSearchHandle, const FOnlineAccountIdHandle& TargetUserId);
+
 	/**
 	 * Writes all values in the passed SessionSettings to the SessionModificationHandle
 	 */
 	void WriteCreateSessionModificationHandle(EOS_HSessionModification& SessionModificationHandle, const FSessionSettings& SessionSettings);
 	/**
-	 * Writes only the new values for all updated settings to the SessionModificationHandle, as well as updating them on the passed SessionSettings
+	 * Writes only the new values for all updated settings to the SessionModificationHandle
 	 */
-	void WriteUpdateSessionModificationHandle(EOS_HSessionModification& SessionModificationHandle, FSessionSettings& SessionSettings, const FSessionSettingsUpdate& NewSettings);
+	void WriteUpdateSessionModificationHandle(EOS_HSessionModification& SessionModificationHandle, const FName& SessionName, const FSessionSettingsUpdate& NewSettings);
 	TFuture<TDefaultErrorResult<FUpdateSessionImpl>> UpdateSessionImpl(FUpdateSessionImpl::Params&& Params);
 
 	void WriteSessionSearchHandle(FSessionSearchHandleEOSGS& SessionSearchHandle, const FFindSessions::Params& Params);
@@ -147,21 +168,22 @@ protected:
 	 */
 	TResult<TSharedRef<const FSession>, FOnlineError> BuildSessionFromUIEvent(const EOS_UI_EventId& UIEventId) const;
 
-protected:
-	FOnlineServicesEOSGS& Services;
+private:
+	/** Returns true if the named session is found and removed successfully, and false otherwise */
+	bool TryRemoveSession(const FName& SessionName);
 
+	// FSessionsLAN
+	virtual void AppendSessionToPacket(FNboSerializeToBuffer& Packet, const FSessionLAN& Session) override;
+	virtual void ReadSessionFromPacket(FNboSerializeFromBuffer& Packet, FSessionLAN& Session) override;
+
+protected:
 	EOS_HSessions SessionsHandle = nullptr;
 
 	EOSEventRegistrationPtr OnSessionInviteReceivedEventRegistration;
 	EOSEventRegistrationPtr OnSessionInviteAcceptedEventRegistration;
 	EOSEventRegistrationPtr OnJoinSessionAcceptedEventRegistration;
 
-	TMap<FName, TSharedRef<FSessionEOSGS>> SessionsByName;
-	TMap<FOnlineSessionIdHandle, TSharedRef<FSessionEOSGS>> SessionsById;
-
 	TSharedPtr<FSessionSearchHandleEOSGS> CurrentSessionSearchHandleEOSGS;
-	TSharedPtr<FFindSessions::Result> CurrentSessionSearch;
-	TSharedPtr<TOnlineAsyncOp<FFindSessions>> CurrentSessionSearchAsyncOpHandle;
 };
 
 /* UE::Online */ }
\ No newline at end of file
diff --git a/Engine/Plugins/Online/Experimental/OnlineServicesNull/Source/Private/Online/LobbyRegistryNull.cpp b/Engine/Plugins/Online/Experimental/OnlineServicesNull/Source/Private/Online/LobbyRegistryNull.cpp
--- a/Engine/Plugins/Online/Experimental/OnlineServicesNull/Source/Private/Online/LobbyRegistryNull.cpp
+++ b/Engine/Plugins/Online/Experimental/OnlineServicesNull/Source/Private/Online/LobbyRegistryNull.cpp
@@ -31,7 +31,7 @@ FOnlineLobbyIdHandle FOnlineLobbyIdRegistryNull::FindOrAdd(FString LobbyId)
 
 UE::Online::FOnlineLobbyIdHandle FOnlineLobbyIdRegistryNull::GetNext()
 {
-	return FindOrAdd(FGuid().ToString());
+	return FindOrAdd(FGuid::NewGuid().ToString());
 }
 
 
diff --git a/Engine/Plugins/Online/Experimental/OnlineServicesNull/Source/Private/Online/SessionsNull.cpp b/Engine/Plugins/Online/Experimental/OnlineServicesNull/Source/Private/Online/SessionsNull.cpp
--- a/Engine/Plugins/Online/Experimental/OnlineServicesNull/Source/Private/Online/SessionsNull.cpp
+++ b/Engine/Plugins/Online/Experimental/OnlineServicesNull/Source/Private/Online/SessionsNull.cpp
@@ -2,27 +2,11 @@
 
 #include "Online/SessionsNull.h"
 
-#include "Misc/Guid.h"
-#include "Online/NboSerializerNullSvc.h"
 #include "Online/OnlineServicesNull.h"
-#include "SocketSubsystem.h"
-#include "UObject/CoreNet.h"
+#include "Online/NboSerializerNullSvc.h"
 
 namespace UE::Online {
 
-/** FOnlineSessionIdRegistryNull */
-
-FOnlineSessionIdRegistryNull& FOnlineSessionIdRegistryNull::Get()
-{
-	static FOnlineSessionIdRegistryNull Instance;
-	return Instance;
-}
-
-FOnlineSessionIdHandle FOnlineSessionIdRegistryNull::GetNextSessionId()
-{
-	return BasicRegistry.FindOrAddHandle(FGuid().ToString());
-}
-
 /** FSessionsNull */
 
 FSessionsNull::FSessionsNull(FOnlineServicesNull& InServices)
@@ -31,367 +15,9 @@ FSessionsNull::FSessionsNull(FOnlineServicesNull& InServices)
 
 }
 
-void FSessionsNull::Tick(float DeltaSeconds)
-{
-	if (LANSessionManager->GetBeaconState() > ELanBeaconState::NotUsingLanBeacon)
-	{
-		LANSessionManager->Tick(DeltaSeconds);
-	}
-
-	Super::Tick(DeltaSeconds);
-}
-
-TOnlineAsyncOpHandle<FCreateSession> FSessionsNull::CreateSession(FCreateSession::Params&& Params)
-{
-	TOnlineAsyncOpRef<FCreateSession> Op = GetOp<FCreateSession>(MoveTemp(Params));
-	const FCreateSession::Params& OpParams = Op->GetParams();
-
-	FOnlineError ParamsCheck = CheckCreateSessionParams(OpParams);
-	if (ParamsCheck != Errors::Success())
-	{
-		Op->SetError(MoveTemp(ParamsCheck));
-		return Op->GetHandle();
-	}
-
-	// We'll only host on the LAN beacon for public sessions
-	if (OpParams.SessionSettings.JoinPolicy == ESessionJoinPolicy::Public)
-	{
-		if (!TryHostLANSession())
-		{
-			Op->SetError(Errors::RequestFailure()); // TODO: May need a new error type
-			return Op->GetHandle();
-		}
-	}
-
-	Op->Then([this](TOnlineAsyncOp<FCreateSession>& Op) mutable
-	{
-		const FCreateSession::Params& OpParams = Op.GetParams();
-
-		FOnlineError StateCheck = CheckCreateSessionState(OpParams);
-		if (StateCheck != Errors::Success())
-		{
-			Op.SetError(MoveTemp(StateCheck));
-			return;
-		}
-
-		FSessionNull NewSession = FSessionNull();
-		NewSession.CurrentState = ESessionState::Valid;
-		NewSession.OwnerUserId = OpParams.LocalUserId;
-		NewSession.SessionId = FOnlineSessionIdRegistryNull::Get().GetNextSessionId();
-		NewSession.SessionSettings = OpParams.SessionSettings;
-
-		// For Null sessions, we'll add all the Session members manually instead of calling JoinSession since there is no API calls involved
-		NewSession.SessionSettings.SessionMembers.Append(OpParams.LocalUsers);
-	
-		TSharedRef<FSessionNull> NewSessionNullRef = MakeShared<FSessionNull>(NewSession);
-
-		SessionsByName.Add(OpParams.SessionName, NewSessionNullRef);
-		SessionsById.Add(NewSession.SessionId, NewSessionNullRef);
-
-		Op.SetResult(FCreateSession::Result{ NewSessionNullRef });
-	})
-	.Enqueue(GetSerialQueue());
-
-	return Op->GetHandle();
-}
-
-TOnlineAsyncOpHandle<FUpdateSession> FSessionsNull::UpdateSession(FUpdateSession::Params&& Params)
-{
-	TOnlineAsyncOpRef<FUpdateSession> Op = GetOp<FUpdateSession>(MoveTemp(Params));
-
-	Op->Then([this](TOnlineAsyncOp<FUpdateSession>& Op) mutable
-	{
-		const FUpdateSession::Params& OpParams = Op.GetParams();
-
-		FOnlineError StateCheck = CheckUpdateSessionState(OpParams);
-		if (StateCheck != Errors::Success())
-		{
-			Op.SetError(MoveTemp(StateCheck));
-			return;
-		}
-
-		// We'll check and update all settings one by one, with additional logic wherever required
-
-		TSharedRef<FSession>& FoundSession = SessionsByName.FindChecked(OpParams.SessionName);
-		FSessionSettings& SessionSettings = FoundSession->SessionSettings;
-		const FSessionSettingsUpdate& UpdatedSettings = OpParams.Mutations;
-
-		if (UpdatedSettings.bAllowNewMembers.IsSet())
-		{
-			SessionSettings.bAllowNewMembers = UpdatedSettings.bAllowNewMembers.GetValue();
-		}
-
-		if (UpdatedSettings.bAllowSanctionedPlayers.IsSet())
-		{
-			SessionSettings.bAllowSanctionedPlayers = UpdatedSettings.bAllowSanctionedPlayers.GetValue();
-		}
-
-		if (UpdatedSettings.bAllowUnregisteredPlayers.IsSet())
-		{
-			SessionSettings.bAllowUnregisteredPlayers = UpdatedSettings.bAllowUnregisteredPlayers.GetValue();
-		}
-
-		if (UpdatedSettings.bAntiCheatProtected.IsSet())
-		{
-			SessionSettings.bAntiCheatProtected = UpdatedSettings.bAntiCheatProtected.GetValue();
-		}
-
-		if (UpdatedSettings.bPresenceEnabled.IsSet())
-		{
-			SessionSettings.bPresenceEnabled = UpdatedSettings.bPresenceEnabled.GetValue();
-		}
-
-		if (UpdatedSettings.IsDedicatedServerSession.IsSet())
-		{
-			SessionSettings.IsDedicatedServerSession = UpdatedSettings.IsDedicatedServerSession.GetValue();
-		}
-
-		if (UpdatedSettings.IsLANSession.IsSet())
-		{
-			SessionSettings.IsLANSession = UpdatedSettings.IsLANSession.GetValue();
-		}
-
-		if (UpdatedSettings.JoinPolicy.IsSet())
-		{
-			// Changes in the join policy setting will mean we start or stop the LAN Beacon broadcasting the session information
-			// There is no FriendsNull interface at the time of this implementation, hence the binary behavior (Public/Non-Public)
-			if (SessionSettings.JoinPolicy != ESessionJoinPolicy::Public && UpdatedSettings.JoinPolicy == ESessionJoinPolicy::Public)
-			{
-				TryHostLANSession();
-			}
-			else if (SessionSettings.JoinPolicy == ESessionJoinPolicy::Public && UpdatedSettings.JoinPolicy != ESessionJoinPolicy::Public)
-			{
-				StopLANSession();
-			}
-
-			SessionSettings.JoinPolicy = UpdatedSettings.JoinPolicy.GetValue();
-		}
-
-		if (UpdatedSettings.NumMaxPrivateConnections.IsSet())
-		{
-			SessionSettings.NumMaxPrivateConnections = UpdatedSettings.NumMaxPrivateConnections.GetValue();
-		}
-
-		if (UpdatedSettings.NumMaxPublicConnections.IsSet())
-		{
-			SessionSettings.NumMaxPublicConnections = UpdatedSettings.NumMaxPublicConnections.GetValue();
-		}
-
-		if (UpdatedSettings.NumOpenPrivateConnections.IsSet())
-		{
-			SessionSettings.NumOpenPrivateConnections = UpdatedSettings.NumOpenPrivateConnections.GetValue();
-		}
-
-		if (UpdatedSettings.NumOpenPublicConnections.IsSet())
-		{
-			SessionSettings.NumOpenPublicConnections = UpdatedSettings.NumOpenPublicConnections.GetValue();
-		}
-
-		// TODO: We may need some additional logic for schema changes
-		if (UpdatedSettings.SchemaName.IsSet())
-		{
-			SessionSettings.SchemaName = UpdatedSettings.SchemaName.GetValue();
-		}
-
-		if (UpdatedSettings.SessionIdOverride.IsSet())
-		{
-			SessionSettings.SessionIdOverride = UpdatedSettings.SessionIdOverride.GetValue();
-		}
-
-		for (const FName& Key : UpdatedSettings.RemovedCustomSettings)
-		{
-			SessionSettings.CustomSettings.Remove(Key);
-		}
-
-		SessionSettings.CustomSettings.Append(UpdatedSettings.UpdatedCustomSettings);
-
-		for (const FOnlineAccountIdHandle& Key : UpdatedSettings.RemovedRegisteredUsers)
-		{
-			SessionSettings.RegisteredUsers.Remove(Key);
-		}
-
-		SessionSettings.RegisteredUsers.Append(UpdatedSettings.UpdatedRegisteredUsers);
-
-		for (const FOnlineAccountIdHandle& Key : UpdatedSettings.RemovedSessionMembers)
-		{
-			SessionSettings.SessionMembers.Remove(Key);
-		}
-
-		for (FSessionMemberUpdatesMap::TConstIterator It(UpdatedSettings.UpdatedSessionMembers); It; ++It)
-		{
-			if (FSessionMember* SessionMember = SessionSettings.SessionMembers.Find(It.Key()))
-			{
-				const FSessionMemberUpdate& SessionMemberUpdate = It.Value();
-
-				SessionMember->MemberSettings.Append(SessionMemberUpdate.UpdatedMemberSettings);
-
-				for (const FName& Key : SessionMemberUpdate.RemovedMemberSettings)
-				{
-					SessionMember->MemberSettings.Remove(Key);
-				}
-			}
-		}
-
-		// We set the result and fire the event
-		Op.SetResult(FUpdateSession::Result{ FoundSession });
-
-		FSessionUpdated SessionUpdatedEvent{ FoundSession, UpdatedSettings };
-		SessionEvents.OnSessionUpdated.Broadcast(SessionUpdatedEvent);
-	})
-	.Enqueue(GetSerialQueue());
-
-	return Op->GetHandle();
-}
-
-TOnlineAsyncOpHandle<FFindSessions> FSessionsNull::FindSessions(FFindSessions::Params&& Params)
-{
-	TOnlineAsyncOpRef<FFindSessions> Op = GetOp<FFindSessions>(MoveTemp(Params));
-	const FFindSessions::Params& OpParams = Op->GetParams();
-
-	FOnlineError ParamsCheck = CheckFindSessionsParams(OpParams);
-	if (ParamsCheck != Errors::Success())
-	{
-		Op->SetError(MoveTemp(ParamsCheck));
-		return Op->GetHandle();
-	}
-
-	Op->Then([this](TOnlineAsyncOp<FFindSessions>& Op) mutable
-	{
-		const FFindSessions::Params& OpParams = Op.GetParams();
-
-		FOnlineError StateCheck = CheckFindSessionsState(OpParams);
-		if (StateCheck != Errors::Success())
-		{
-			Op.SetError(MoveTemp(StateCheck));
-			return;
-		}
-
-		CurrentSessionSearch = MakeShared<FFindSessions::Result>();
-		CurrentSessionSearchHandle = Op.AsShared();
-
-		FindLANSessions();
-	})
-	.Enqueue(GetSerialQueue());
-
-	return Op->GetHandle();
-}
-
-TOnlineAsyncOpHandle<FJoinSession> FSessionsNull::JoinSession(FJoinSession::Params&& Params)
-{
-	TOnlineAsyncOpRef<FJoinSession> Op = GetOp<FJoinSession>(MoveTemp(Params));
-	const FJoinSession::Params& OpParams = Op->GetParams();
-
-	FOnlineError ParamsCheck = CheckJoinSessionParams(OpParams);
-	if (ParamsCheck != Errors::Success())
-	{
-		Op->SetError(MoveTemp(ParamsCheck));
-		return Op->GetHandle();
-	}
-
-	// If no restrictions apply, we'll save our copy of the session, add the new players, and register it
-	Op->Then([this](TOnlineAsyncOp<FJoinSession>& Op) mutable
-	{
-		const FJoinSession::Params& OpParams = Op.GetParams();
-
-		FOnlineError StateCheck = CheckJoinSessionState(OpParams);
-		if (StateCheck != Errors::Success())
-		{
-			Op.SetError(MoveTemp(StateCheck));
-			return;
-		}
-
-		TSharedRef<FSessionNull> NewSessionNullRef = MakeShared<FSessionNull>(FSessionNull::Cast(*OpParams.Session));
-		FSessionSettings NewSessionNullSettings = NewSessionNullRef->SessionSettings;
-
-		// TODO: This will move to AddSessionMember
-		TArray<FOnlineAccountIdHandle> LocalUserIds;
-		LocalUserIds.Reserve(OpParams.LocalUsers.Num());
-		OpParams.LocalUsers.GenerateKeyArray(LocalUserIds);
-
-		TArray<FOnlineAccountIdHandle> JoinedUsers;
-
-		for (const FOnlineAccountIdHandle& LocalUserId : LocalUserIds)
-		{
-			bool bReserveSlot = true;
-
-			if (const FRegisteredUser* RegisteredUser = NewSessionNullSettings.RegisteredUsers.Find(LocalUserId))
-			{
-				bReserveSlot = !RegisteredUser->bHasReservedSlot;
-			}
-
-			if (bReserveSlot)
-			{
-				if (NewSessionNullSettings.NumOpenPublicConnections > 0)
-				{
-					NewSessionNullSettings.NumOpenPublicConnections--;
-				}
-				else if (NewSessionNullSettings.NumOpenPrivateConnections > 0)
-				{
-					NewSessionNullSettings.NumOpenPrivateConnections--;
-				}
-			}
-
-			JoinedUsers.Add(LocalUserId);
-		}
-
-		FSessionJoined SessionJoinedEvent = { MoveTemp(JoinedUsers), NewSessionNullRef };
-
-		NewSessionNullSettings.SessionMembers.Append(OpParams.LocalUsers);
-
-		SessionsByName.Add(OpParams.SessionName, NewSessionNullRef);
-		SessionsById.Add(NewSessionNullRef->SessionId, NewSessionNullRef);
-
-		Op.SetResult(FJoinSession::Result{ NewSessionNullRef });
-
-		SessionEvents.OnSessionJoined.Broadcast(SessionJoinedEvent);
-	})
-	.Enqueue(GetSerialQueue());
-
-	return Op->GetHandle();
-}
-
-TOnlineAsyncOpHandle<FLeaveSession> FSessionsNull::LeaveSession(FLeaveSession::Params&& Params)
-{
-	TOnlineAsyncOpRef<FLeaveSession> Op = GetOp<FLeaveSession>(MoveTemp(Params));
-
-	Op->Then([this](TOnlineAsyncOp<FLeaveSession>& Op) mutable
-	{
-		const FLeaveSession::Params& OpParams = Op.GetParams();
-
-		FOnlineError StateCheck = CheckLeaveSessionState(OpParams);
-		if (StateCheck != Errors::Success())
-		{
-			Op.SetError(MoveTemp(StateCheck));
-			return;
-		}
-
-		TSharedRef<FSession>& FoundSession = SessionsByName.FindChecked(OpParams.SessionName);
-		
-		// Only the host should stop the session at this point
-		if (FoundSession->OwnerUserId == OpParams.LocalUserId && FoundSession->SessionSettings.JoinPolicy == ESessionJoinPolicy::Public)
-		{
-			StopLANSession();
-		}
-
-		// TODO: Remove all local users from session
-
-		SessionsById.Remove(FoundSession->SessionId);
-		SessionsByName.Remove(OpParams.SessionName);
-
-		Op.SetResult(FLeaveSession::Result{ });
-
-		FSessionLeft SessionLeftEvent;
-		SessionLeftEvent.LocalUserIds.Append(OpParams.LocalUsers);
-		SessionEvents.OnSessionLeft.Broadcast(SessionLeftEvent);
-	})
-	.Enqueue(GetSerialQueue());
-
-	return Op->GetHandle();
-}
-
 /** FSessionsLAN */
 
-void FSessionsNull::AppendSessionToPacket(FNboSerializeToBuffer& Packet, const FSessionNull& Session)
+void FSessionsNull::AppendSessionToPacket(FNboSerializeToBuffer& Packet, const FSessionLAN& Session)
 {
 	using namespace NboSerializerNullSvc;
 
@@ -403,7 +29,7 @@ void FSessionsNull::AppendSessionToPacket(FNboSerializeToBuffer& Packet, const F
 	// TODO: Write session settings to packet, after SchemaVariant work
 }
 
-void FSessionsNull::ReadSessionFromPacket(FNboSerializeFromBuffer& Packet, FSessionNull& Session)
+void FSessionsNull::ReadSessionFromPacket(FNboSerializeFromBuffer& Packet, FSessionLAN& Session)
 {
 	using namespace NboSerializerNullSvc;
 
diff --git a/Engine/Plugins/Online/Experimental/OnlineServicesNull/Source/Public/Online/NboSerializerNullSvc.h b/Engine/Plugins/Online/Experimental/OnlineServicesNull/Source/Public/Online/NboSerializerNullSvc.h
--- a/Engine/Plugins/Online/Experimental/OnlineServicesNull/Source/Public/Online/NboSerializerNullSvc.h
+++ b/Engine/Plugins/Online/Experimental/OnlineServicesNull/Source/Public/Online/NboSerializerNullSvc.h
@@ -6,7 +6,7 @@
 
 #include "Online/AuthNull.h"
 #include "Online/LobbiesNull.h"
-#include "Online/SessionsNull.h"
+#include "Online/SessionsLAN.h"
 #include "Online/CoreOnline.h"
 
 /**
@@ -27,7 +27,7 @@ inline void SerializeToBuffer(FNboSerializeToBuffer& Ar, const FOnlineAccountIdH
 
 inline void SerializeToBuffer(FNboSerializeToBuffer& Ar, const FOnlineSessionIdHandle& SessionId)
 {
-	TArray<uint8> Data = FOnlineSessionIdRegistryNull::Get().ToReplicationData(SessionId);
+	TArray<uint8> Data = FOnlineSessionIdRegistryLAN::Get().ToReplicationData(SessionId);
 	Ar << Data.Num();
 	Ar.WriteBinary(Data.GetData(), Data.Num());
 }
@@ -61,7 +61,7 @@ inline void SerializeFromBuffer(FNboSerializeFromBuffer& Ar, FOnlineSessionIdHan
 	Ar >> Size;
 	Ar.ReadBinaryArray(Data, Size);
 
-	SessionId = FOnlineSessionIdRegistryNull::Get().FromReplicationData(Data);
+	SessionId = FOnlineSessionIdRegistryLAN::Get().FromReplicationData(Data);
 }
 
 inline void SerializeFromBuffer(FNboSerializeFromBuffer& Ar, TMap<FLobbyAttributeId, FLobbyVariant>& Map)
diff --git a/Engine/Plugins/Online/Experimental/OnlineServicesNull/Source/Public/Online/OnlineServicesNull.h b/Engine/Plugins/Online/Experimental/OnlineServicesNull/Source/Public/Online/OnlineServicesNull.h
--- a/Engine/Plugins/Online/Experimental/OnlineServicesNull/Source/Public/Online/OnlineServicesNull.h
+++ b/Engine/Plugins/Online/Experimental/OnlineServicesNull/Source/Public/Online/OnlineServicesNull.h
@@ -6,8 +6,6 @@
 
 namespace UE::Online {
 
-static FName CONNECT_STRING_TAG = TEXT("CONNECT_STRING");
-
 class ONLINESERVICESNULL_API FOnlineServicesNull : public FOnlineServicesCommon
 {
 public:
diff --git a/Engine/Plugins/Online/Experimental/OnlineServicesNull/Source/Public/Online/SessionsNull.h b/Engine/Plugins/Online/Experimental/OnlineServicesNull/Source/Public/Online/SessionsNull.h
--- a/Engine/Plugins/Online/Experimental/OnlineServicesNull/Source/Public/Online/SessionsNull.h
+++ b/Engine/Plugins/Online/Experimental/OnlineServicesNull/Source/Public/Online/SessionsNull.h
@@ -1,22 +1,12 @@
 // Copyright Epic Games, Inc. All Rights Reserved.
 #pragma once
 
-#include "CoreMinimal.h"
-#include "Online/LANBeacon.h"
 #include "Online/SessionsLAN.h"
 
 namespace UE::Online {
 
 class FOnlineServicesNull;
 
-class FOnlineSessionIdRegistryNull : public TOnlineSessionIdStringRegistry<EOnlineServices::Null>
-{
-public:
-	static FOnlineSessionIdRegistryNull& Get();
-
-	FOnlineSessionIdHandle GetNextSessionId();
-};
-
 typedef FSessionLAN FSessionNull;
 
 class FSessionsNull : public FSessionsLAN
@@ -27,15 +17,6 @@ public:
 	FSessionsNull(FOnlineServicesNull& InServices);
 	virtual ~FSessionsNull() = default;
 
-	virtual void Tick(float DeltaSeconds) override;
-
-	// ISessions
-	virtual TOnlineAsyncOpHandle<FCreateSession> CreateSession(FCreateSession::Params&& Params) override;
-	virtual TOnlineAsyncOpHandle<FUpdateSession> UpdateSession(FUpdateSession::Params&& Params) override;
-	virtual TOnlineAsyncOpHandle<FFindSessions> FindSessions(FFindSessions::Params&& Params) override;
-	virtual TOnlineAsyncOpHandle<FJoinSession> JoinSession(FJoinSession::Params&& Params) override;
-	virtual TOnlineAsyncOpHandle<FLeaveSession> LeaveSession(FLeaveSession::Params&& Params) override;
-
 private:
 	// FSessionsLAN
 	virtual void AppendSessionToPacket(FNboSerializeToBuffer& Packet, const FSessionLAN& Session) override;
