OSSv2: Add Leaderboards EOS implementation.

diff --git a/Engine/Plugins/Online/Experimental/OnlineServices/Source/OnlineServicesInterface/Public/Online/Leaderboards.h b/Engine/Plugins/Online/Experimental/OnlineServices/Source/OnlineServicesInterface/Public/Online/Leaderboards.h
--- a/Engine/Plugins/Online/Experimental/OnlineServices/Source/OnlineServicesInterface/Public/Online/Leaderboards.h
+++ b/Engine/Plugins/Online/Experimental/OnlineServices/Source/OnlineServicesInterface/Public/Online/Leaderboards.h
@@ -6,6 +6,8 @@
 #include "Online/OnlineAsyncOpHandle.h"
 #include "Online/OnlineMeta.h"
 
+#define UE_LEADERBOARD_RANK_UNKNOWN -1
+
 namespace UE::Online {
 
 struct FLeaderboardEntry
@@ -24,6 +26,8 @@ struct FReadEntriesForUsers
 
 	struct Params
 	{
+		/* Local user id */
+		FOnlineAccountIdHandle LocalUserId;
 		/* The account ids of users */
 		TArray<FOnlineAccountIdHandle> UserIds;
 		/* The leaderboard name */
@@ -43,10 +47,12 @@ struct FReadEntriesAroundRank
 
 	struct Params
 	{
+		/* Local user id */
+		FOnlineAccountIdHandle LocalUserId;
 		/* The rank index to start read */
-		int32 Rank;
+		uint32 Rank;
 		/* How many entries to fetch */
-		int32 Limit;
+		uint32 Limit;
 		/* The leaderboard to read */
 		FString BoardName;
 	};
@@ -64,12 +70,14 @@ struct FReadEntriesAroundUser
 
 	struct Params
 	{
+		/* Local user id */
+		FOnlineAccountIdHandle LocalUserId;
 		/* The account id of specified user */
 		FOnlineAccountIdHandle UserId;
 		/* The offset to the rank of user */
 		int32 Offset;
 		/* How many entries to fetch */
-		int32 Limit;
+		uint32 Limit;
 		/* The leaderboard to read */
 		FString BoardName;
 	};
@@ -117,6 +125,7 @@ BEGIN_ONLINE_STRUCT_META(FLeaderboardEntry)
 END_ONLINE_STRUCT_META()
 
 BEGIN_ONLINE_STRUCT_META(FReadEntriesForUsers::Params)
+	ONLINE_STRUCT_FIELD(FReadEntriesForUsers::Params, LocalUserId),
 	ONLINE_STRUCT_FIELD(FReadEntriesForUsers::Params, UserIds),
 	ONLINE_STRUCT_FIELD(FReadEntriesForUsers::Params, BoardName)
 END_ONLINE_STRUCT_META()
@@ -126,6 +135,7 @@ BEGIN_ONLINE_STRUCT_META(FReadEntriesForUsers::Result)
 END_ONLINE_STRUCT_META()
 
 BEGIN_ONLINE_STRUCT_META(FReadEntriesAroundRank::Params)
+	ONLINE_STRUCT_FIELD(FReadEntriesAroundRank::Params, LocalUserId),
 	ONLINE_STRUCT_FIELD(FReadEntriesAroundRank::Params, Rank),
 	ONLINE_STRUCT_FIELD(FReadEntriesAroundRank::Params, Limit),
 	ONLINE_STRUCT_FIELD(FReadEntriesAroundRank::Params, BoardName)
@@ -136,6 +146,7 @@ BEGIN_ONLINE_STRUCT_META(FReadEntriesAroundRank::Result)
 END_ONLINE_STRUCT_META()
 
 BEGIN_ONLINE_STRUCT_META(FReadEntriesAroundUser::Params)
+	ONLINE_STRUCT_FIELD(FReadEntriesAroundUser::Params, LocalUserId),
 	ONLINE_STRUCT_FIELD(FReadEntriesAroundUser::Params, UserId),
 	ONLINE_STRUCT_FIELD(FReadEntriesAroundUser::Params, Offset),
 	ONLINE_STRUCT_FIELD(FReadEntriesAroundUser::Params, Limit),
diff --git a/Engine/Plugins/Online/Experimental/OnlineServicesEOSGS/Source/Private/Online/LeaderboardsEOSGS.cpp b/Engine/Plugins/Online/Experimental/OnlineServicesEOSGS/Source/Private/Online/LeaderboardsEOSGS.cpp
new file mode 100644
--- /dev/null
+++ b/Engine/Plugins/Online/Experimental/OnlineServicesEOSGS/Source/Private/Online/LeaderboardsEOSGS.cpp
@@ -0,0 +1,257 @@
+// Copyright Epic Games, Inc. All Rights Reserved.
+
+#include "Online/LeaderboardsEOSGS.h"
+#include "Online/AuthEOSGS.h"
+#include "Online/OnlineIdEOSGS.h"
+#include "Online/OnlineServicesEOSGS.h"
+#include "Online/OnlineServicesEOSGSTypes.h"
+#include "Math/UnrealMathUtility.h"
+
+#include "eos_leaderboards.h"
+
+namespace UE::Online {
+
+#define EOS_LEADERBOARD_STAT_NAME_MAX_LENGTH 256 + 1 // 256 plus null terminator
+#define EOS_LEADERBOARD_MAX_ENTRIES 1000
+
+FLeaderboardsEOSGS::FLeaderboardsEOSGS(FOnlineServicesEOSGS& InServices)
+	: Super(InServices)
+{
+}
+
+void FLeaderboardsEOSGS::Initialize()
+{
+	Super::Initialize();
+
+	LeaderboardsHandle = EOS_Platform_GetLeaderboardsInterface(static_cast<FOnlineServicesEOSGS&>(GetServices()).GetEOSPlatformHandle());
+	check(LeaderboardsHandle);
+}
+
+TOnlineAsyncOpHandle<FReadEntriesForUsers> FLeaderboardsEOSGS::ReadEntriesForUsers(FReadEntriesForUsers::Params&& Params)
+{
+	TOnlineAsyncOpRef<FReadEntriesForUsers> Op = GetOp<FReadEntriesForUsers>(MoveTemp(Params));
+
+	if (!Services.Get<FAuthEOSGS>()->IsLoggedIn(Op->GetParams().LocalUserId))
+	{
+		Op->SetError(Errors::InvalidUser());
+		return Op->GetHandle();
+	}
+
+	Op->Then([this](TOnlineAsyncOp<FReadEntriesForUsers>& InAsyncOp, TPromise<const EOS_Leaderboards_OnQueryLeaderboardUserScoresCompleteCallbackInfo*>&& Promise)
+	{
+		TArray<EOS_ProductUserId> ProductUserIds;
+		ProductUserIds.AddZeroed(InAsyncOp.GetParams().UserIds.Num());
+		for (const FOnlineAccountIdHandle& UserId : InAsyncOp.GetParams().UserIds)
+		{
+			ProductUserIds.Emplace(GetProductUserIdChecked(UserId));
+		}
+
+		EOS_Leaderboards_UserScoresQueryStatInfo StatInfo;
+		StatInfo.ApiVersion = EOS_LEADERBOARDS_USERSCORESQUERYSTATINFO_API_LATEST;
+		static_assert(EOS_LEADERBOARDS_USERSCORESQUERYSTATINFO_API_LATEST == 1, "EOS_Leaderboards_UserScoresQueryStatInfo updated, check new fields");
+
+		char StatNameANSI[EOS_LEADERBOARD_STAT_NAME_MAX_LENGTH];
+		FCStringAnsi::Strncpy(StatNameANSI, TCHAR_TO_UTF8(*InAsyncOp.GetParams().BoardName), EOS_LEADERBOARD_STAT_NAME_MAX_LENGTH);
+
+		StatInfo.StatName = StatNameANSI;
+		StatInfo.Aggregation = EOS_ELeaderboardAggregation::EOS_LA_Latest; // TODO: Use Stats definitions
+
+		EOS_Leaderboards_QueryLeaderboardUserScoresOptions Options = { };
+		Options.ApiVersion = EOS_LEADERBOARDS_QUERYLEADERBOARDUSERSCORES_API_LATEST;
+		static_assert(EOS_LEADERBOARDS_QUERYLEADERBOARDUSERSCORES_API_LATEST == 2, "EOS_Leaderboards_QueryLeaderboardUserScores updated, check new fields");
+		Options.UserIds = ProductUserIds.GetData();
+		Options.UserIdsCount = ProductUserIds.Num();
+		Options.StatInfo = &StatInfo;
+		Options.StatInfoCount = 1;
+		Options.StartTime = EOS_LEADERBOARDS_TIME_UNDEFINED;
+		Options.EndTime = EOS_LEADERBOARDS_TIME_UNDEFINED;
+		Options.LocalUserId = GetProductUserIdChecked(InAsyncOp.GetParams().LocalUserId);
+
+		EOS_Async(EOS_Leaderboards_QueryLeaderboardUserScores, LeaderboardsHandle, Options, MoveTemp(Promise));
+	})
+	.Then([this](TOnlineAsyncOp<FReadEntriesForUsers>& InAsyncOp, const EOS_Leaderboards_OnQueryLeaderboardUserScoresCompleteCallbackInfo* Data)
+	{
+		if (Data->ResultCode != EOS_EResult::EOS_Success)
+		{
+			UE_LOG(LogTemp, Warning, TEXT("EOS_Leaderboards_IngestStat failed with result=[%s]"), *LexToString(Data->ResultCode));
+			InAsyncOp.SetError(FromEOSError(Data->ResultCode));
+			return;
+		}
+
+		char StatNameANSI[EOS_LEADERBOARD_STAT_NAME_MAX_LENGTH];
+		FCStringAnsi::Strncpy(StatNameANSI, TCHAR_TO_UTF8(*InAsyncOp.GetParams().BoardName), EOS_LEADERBOARD_STAT_NAME_MAX_LENGTH);
+
+		FReadEntriesForUsers::Result Result;
+
+		for (const FOnlineAccountIdHandle& UserId : InAsyncOp.GetParams().UserIds)
+		{
+			EOS_Leaderboards_CopyLeaderboardUserScoreByUserIdOptions UserCopyOptions = { };
+			UserCopyOptions.ApiVersion = EOS_LEADERBOARDS_COPYLEADERBOARDUSERSCOREBYUSERID_API_LATEST;
+			static_assert(EOS_LEADERBOARDS_COPYLEADERBOARDUSERSCOREBYUSERID_API_LATEST == 1, "EOS_Leaderboards_CopyLeaderboardUserScoreByUserIdOptions updated, check new fields");
+
+			UserCopyOptions.UserId = GetProductUserIdChecked(UserId);
+			UserCopyOptions.StatName = StatNameANSI;
+
+			EOS_Leaderboards_LeaderboardUserScore* LeaderboardUserScore = nullptr;
+			EOS_EResult CopyResult = EOS_Leaderboards_CopyLeaderboardUserScoreByUserId(LeaderboardsHandle, &UserCopyOptions, &LeaderboardUserScore);
+
+			if (CopyResult != EOS_EResult::EOS_Success)
+			{
+				UE_LOG(LogTemp, VeryVerbose, TEXT("Value not found for leaderboard %s."), *InAsyncOp.GetParams().BoardName);
+				continue;
+			}
+
+			FLeaderboardEntry& Entry = Result.Entries.Emplace_GetRef();
+			Entry.Rank = UE_LEADERBOARD_RANK_UNKNOWN;
+			Entry.UserId = UserId;
+			Entry.Score = LeaderboardUserScore->Score;
+
+			EOS_Leaderboards_LeaderboardUserScore_Release(LeaderboardUserScore);
+		}
+
+		InAsyncOp.SetResult(MoveTemp(Result));
+	}) 
+	.Enqueue(GetSerialQueue());
+
+	return Op->GetHandle();
+}
+
+namespace Private
+{
+
+void QueryLeaderboardsEOS(EOS_HLeaderboards LeaderboardsHandle, const FOnlineAccountIdHandle& LocalUserId, const FString& BoardName, TPromise<const EOS_Leaderboards_OnQueryLeaderboardRanksCompleteCallbackInfo*>&& Promise)
+{
+	EOS_Leaderboards_QueryLeaderboardRanksOptions Options;
+	Options.LocalUserId = GetProductUserIdChecked(LocalUserId);
+	Options.ApiVersion = EOS_LEADERBOARDS_QUERYLEADERBOARDRANKS_API_LATEST;
+	static_assert(EOS_LEADERBOARDS_QUERYLEADERBOARDRANKS_API_LATEST == 2, "EOS_Leaderboards_QueryLeaderboardRanks updated, check new fields");
+
+	char LeaderboardNameANSI[EOS_LEADERBOARD_STAT_NAME_MAX_LENGTH];
+	FCStringAnsi::Strncpy(LeaderboardNameANSI, TCHAR_TO_UTF8(*BoardName), EOS_LEADERBOARD_STAT_NAME_MAX_LENGTH);
+	Options.LeaderboardId = LeaderboardNameANSI;
+
+	EOS_Async(EOS_Leaderboards_QueryLeaderboardRanks, LeaderboardsHandle, Options, MoveTemp(Promise));
+}
+
+void ReadEntriesInRange(EOS_HLeaderboards LeaderboardsHandle, uint32 StartIndex, uint32 EndIndex, TArray<FLeaderboardEntry>& OutEntries)
+{
+	for (uint32 i = StartIndex; i <= EndIndex; ++i)
+	{
+		EOS_Leaderboards_CopyLeaderboardRecordByIndexOptions CopyRecordOptions;
+		CopyRecordOptions.ApiVersion = EOS_LEADERBOARDS_COPYLEADERBOARDRECORDBYINDEX_API_LATEST;
+		static_assert(EOS_LEADERBOARDS_COPYLEADERBOARDRECORDBYINDEX_API_LATEST == 2, "EOS_Leaderboards_CopyLeaderboardRecordByIndexOptions updated, check new fields");
+		CopyRecordOptions.LeaderboardRecordIndex = i;
+
+		EOS_Leaderboards_LeaderboardRecord* LeaderboardRecord = nullptr;
+		EOS_EResult CopyResult = EOS_Leaderboards_CopyLeaderboardRecordByIndex(LeaderboardsHandle, &CopyRecordOptions, &LeaderboardRecord);
+		if (CopyResult == EOS_EResult::EOS_Success)
+		{
+			FLeaderboardEntry& LeaderboardEntry = OutEntries.Emplace_GetRef();
+			LeaderboardEntry.UserId = FindAccountIdChecked(LeaderboardRecord->UserId);
+			LeaderboardEntry.Rank = LeaderboardRecord->Rank;
+			LeaderboardEntry.Score = LeaderboardRecord->Score;
+
+			EOS_Leaderboards_LeaderboardRecord_Release(LeaderboardRecord);
+		}
+	}
+}
+
+}
+
+TOnlineAsyncOpHandle<FReadEntriesAroundRank> FLeaderboardsEOSGS::ReadEntriesAroundRank(FReadEntriesAroundRank::Params&& Params)
+{
+	TOnlineAsyncOpRef<FReadEntriesAroundRank> Op = GetOp<FReadEntriesAroundRank>(MoveTemp(Params));
+
+	if (!Services.Get<FAuthEOSGS>()->IsLoggedIn(Op->GetParams().LocalUserId))
+	{
+		Op->SetError(Errors::InvalidUser());
+		return Op->GetHandle();
+	}
+
+	if (Op->GetParams().Rank >= EOS_LEADERBOARD_MAX_ENTRIES || Op->GetParams().Limit == 0)
+	{
+		Op->SetError(Errors::InvalidParams());
+		return Op->GetHandle();
+	}
+
+	Op->Then([this](TOnlineAsyncOp<FReadEntriesAroundRank>& InAsyncOp, TPromise<const EOS_Leaderboards_OnQueryLeaderboardRanksCompleteCallbackInfo*>&& Promise)
+	{
+		Private::QueryLeaderboardsEOS(LeaderboardsHandle, InAsyncOp.GetParams().LocalUserId, InAsyncOp.GetParams().BoardName, MoveTemp(Promise));
+	})
+	.Then([this](TOnlineAsyncOp<FReadEntriesAroundRank>& InAsyncOp, const EOS_Leaderboards_OnQueryLeaderboardRanksCompleteCallbackInfo* Data)
+	{
+		if (Data->ResultCode != EOS_EResult::EOS_Success)
+		{
+			UE_LOG(LogTemp, Warning, TEXT("EOS_Leaderboards_QueryLeaderboardRanks failed with result=[%s]"), *LexToString(Data->ResultCode));
+			InAsyncOp.SetError(FromEOSError(Data->ResultCode));
+			return;
+		}
+
+		FReadEntriesAroundRank::Result Result;
+		int32 StartIndex = InAsyncOp.GetParams().Rank;
+		int32 EndIndex = StartIndex + InAsyncOp.GetParams().Limit;
+		EndIndex = FMath::Clamp(EndIndex, 0, EOS_LEADERBOARD_MAX_ENTRIES-1);
+		Private::ReadEntriesInRange(LeaderboardsHandle, StartIndex, EndIndex, Result.Entries);
+		InAsyncOp.SetResult(MoveTemp(Result));
+	})
+	.Enqueue(GetSerialQueue());
+
+	return Op->GetHandle();
+}
+
+TOnlineAsyncOpHandle<FReadEntriesAroundUser> FLeaderboardsEOSGS::ReadEntriesAroundUser(FReadEntriesAroundUser::Params&& Params)
+{
+	TOnlineAsyncOpRef<FReadEntriesAroundUser> Op = GetOp<FReadEntriesAroundUser>(MoveTemp(Params));
+
+	if (!Services.Get<FAuthEOSGS>()->IsLoggedIn(Op->GetParams().LocalUserId))
+	{
+		Op->SetError(Errors::InvalidUser());
+		return Op->GetHandle();
+	}
+
+	if (Op->GetParams().Limit > EOS_LEADERBOARD_MAX_ENTRIES || Op->GetParams().Limit == 0 || Op->GetParams().Offset >= EOS_LEADERBOARD_MAX_ENTRIES || Op->GetParams().Offset <= -EOS_LEADERBOARD_MAX_ENTRIES)
+	{
+		Op->SetError(Errors::InvalidParams());
+		return Op->GetHandle();
+	}
+
+	Op->Then([this](TOnlineAsyncOp<FReadEntriesAroundUser>& InAsyncOp, TPromise<const EOS_Leaderboards_OnQueryLeaderboardRanksCompleteCallbackInfo*>&& Promise)
+	{
+		Private::QueryLeaderboardsEOS(LeaderboardsHandle, InAsyncOp.GetParams().LocalUserId, InAsyncOp.GetParams().BoardName, MoveTemp(Promise));
+	})
+	.Then([this](TOnlineAsyncOp<FReadEntriesAroundUser>& InAsyncOp, const EOS_Leaderboards_OnQueryLeaderboardRanksCompleteCallbackInfo* Data)
+	{
+		if (Data->ResultCode != EOS_EResult::EOS_Success)
+		{
+			UE_LOG(LogTemp, Warning, TEXT("EOS_Leaderboards_QueryLeaderboardRanks failed with result=[%s]"), *LexToString(Data->ResultCode));
+			InAsyncOp.SetError(FromEOSError(Data->ResultCode));
+			return;
+		}
+
+		FReadEntriesAroundUser::Result Result;
+
+		EOS_Leaderboards_LeaderboardRecord* LeaderboardRecord = nullptr;
+		EOS_Leaderboards_CopyLeaderboardRecordByUserIdOptions CopyRecordByUserIdOptions;
+		CopyRecordByUserIdOptions.ApiVersion = EOS_LEADERBOARDS_COPYLEADERBOARDRECORDBYUSERID_API_LATEST;
+		static_assert(EOS_LEADERBOARDS_COPYLEADERBOARDRECORDBYUSERID_API_LATEST == 2, "EOS_Leaderboards_CopyLeaderboardRecordByUserIdOptions updated, check new fields");
+		CopyRecordByUserIdOptions.UserId = GetProductUserIdChecked(InAsyncOp.GetParams().UserId);
+		EOS_EResult CopyResult = EOS_Leaderboards_CopyLeaderboardRecordByUserId(LeaderboardsHandle, &CopyRecordByUserIdOptions, &LeaderboardRecord);
+		if (CopyResult == EOS_EResult::EOS_Success)
+		{
+			int32 StartIndex = (int32)LeaderboardRecord->Rank + InAsyncOp.GetParams().Offset;
+			StartIndex = FMath::Clamp(StartIndex, 0, EOS_LEADERBOARD_MAX_ENTRIES-1);
+			int32 EndIndex = StartIndex + InAsyncOp.GetParams().Limit;
+			EndIndex = FMath::Clamp(EndIndex, 0, EOS_LEADERBOARD_MAX_ENTRIES-1);
+			Private::ReadEntriesInRange(LeaderboardsHandle, StartIndex, EndIndex, Result.Entries);
+
+			EOS_Leaderboards_LeaderboardRecord_Release(LeaderboardRecord);
+		}
+
+		InAsyncOp.SetResult(MoveTemp(Result));
+	})
+	.Enqueue(GetSerialQueue());
+
+	return Op->GetHandle();
+}
+
+/* UE::Online */ }
diff --git a/Engine/Plugins/Online/Experimental/OnlineServicesEOSGS/Source/Private/Online/OnlineServicesEOSGS.cpp b/Engine/Plugins/Online/Experimental/OnlineServicesEOSGS/Source/Private/Online/OnlineServicesEOSGS.cpp
--- a/Engine/Plugins/Online/Experimental/OnlineServicesEOSGS/Source/Private/Online/OnlineServicesEOSGS.cpp
+++ b/Engine/Plugins/Online/Experimental/OnlineServicesEOSGS/Source/Private/Online/OnlineServicesEOSGS.cpp
@@ -4,6 +4,7 @@
 
 #include "Online/AchievementsEOSGS.h"
 #include "Online/AuthEOSGS.h"
+#include "Online/LeaderboardsEOSGS.h"
 #include "Online/LobbiesEOSGS.h"
 #include "Online/StatsEOSGS.h"
 #include "Online/SessionsEOSGS.h"
@@ -29,6 +30,7 @@ void FOnlineServicesEOSGS::RegisterComponents()
 {
 	Components.Register<FAchievementsEOSGS>(*this);
 	Components.Register<FAuthEOSGS>(*this);
+	Components.Register<FLeaderboardsEOSGS>(*this);
 	Components.Register<FLobbiesEOSGS>(*this);
 	Components.Register<FStatsEOSGS>(*this);
 	Components.Register<FSessionsEOSGS>(*this);
diff --git a/Engine/Plugins/Online/Experimental/OnlineServicesEOSGS/Source/Public/Online/LeaderboardsEOSGS.h b/Engine/Plugins/Online/Experimental/OnlineServicesEOSGS/Source/Public/Online/LeaderboardsEOSGS.h
new file mode 100644
--- /dev/null
+++ b/Engine/Plugins/Online/Experimental/OnlineServicesEOSGS/Source/Public/Online/LeaderboardsEOSGS.h
@@ -0,0 +1,36 @@
+// Copyright Epic Games, Inc. All Rights Reserved.
+
+#pragma once
+
+#include "CoreMinimal.h"
+#include "Online/LeaderboardsCommon.h"
+
+#if defined(EOS_PLATFORM_BASE_FILE_NAME)
+#include EOS_PLATFORM_BASE_FILE_NAME
+#endif
+#include "eos_leaderboards_types.h"
+
+namespace UE::Online {
+
+class FOnlineServicesEOSGS;
+
+class ONLINESERVICESEOSGS_API FLeaderboardsEOSGS : public FLeaderboardsCommon
+{
+public:
+	using Super = FLeaderboardsCommon;
+
+	FLeaderboardsEOSGS(FOnlineServicesEOSGS& InOwningSubsystem);
+
+	// TOnlineComponent
+	virtual void Initialize() override;
+
+	// ILeaderboards
+	virtual TOnlineAsyncOpHandle<FReadEntriesForUsers> ReadEntriesForUsers(FReadEntriesForUsers::Params&& Params) override;
+	virtual TOnlineAsyncOpHandle<FReadEntriesAroundRank> ReadEntriesAroundRank(FReadEntriesAroundRank::Params&& Params) override;
+	virtual TOnlineAsyncOpHandle<FReadEntriesAroundUser> ReadEntriesAroundUser(FReadEntriesAroundUser::Params&& Params) override;
+
+private:
+	EOS_HLeaderboards LeaderboardsHandle = nullptr;
+};
+
+/* UE::Online */ }
