Fix issue detected by static analysis

diff --git a/Engine/Plugins/Online/Experimental/OnlineServicesEOSGS/Source/Private/Online/AuthEOSGS.cpp b/Engine/Plugins/Online/Experimental/OnlineServicesEOSGS/Source/Private/Online/AuthEOSGS.cpp
--- a/Engine/Plugins/Online/Experimental/OnlineServicesEOSGS/Source/Private/Online/AuthEOSGS.cpp
+++ b/Engine/Plugins/Online/Experimental/OnlineServicesEOSGS/Source/Private/Online/AuthEOSGS.cpp
@@ -132,8 +132,7 @@ FEOSConnectLoginOptions& FEOSConnectLoginOptions::operator=(FEOSConnectLoginOpti
 		CredentialsData = Other.CredentialsData;
 		CredentialsData.Token = ExternalAuthTokenUtf8.GetData();
 		Credentials = &CredentialsData;
-		Other.CredentialsData = EOS_Connect_Credentials({});
-		memset(&Other.CredentialsData, 0, sizeof(Other.CredentialsData));
+		Other.CredentialsData = {};
 	}
 
 	if (Other.UserLoginInfo)
@@ -142,7 +141,8 @@ FEOSConnectLoginOptions& FEOSConnectLoginOptions::operator=(FEOSConnectLoginOpti
 		UserLoginInfoData = Other.UserLoginInfoData;
 		UserLoginInfoData.DisplayName = DisplayNameUtf8;
 		UserLoginInfo = &UserLoginInfoData;
-		memset(&Other.UserLoginInfo, 0, sizeof(Other.UserLoginInfo));
+		Other.UserLoginInfo = nullptr;
+		Other.UserLoginInfoData = {};
 	}
 
 	return *this;
