Consolidate EOS error mapping code.
Fix error categories using the same number.
Enable bAddUserNumToNullId for OnlineTests to prevent different FPlatformUserIds from mapping to the same FOnlineAccountIdHandle.

diff --git a/Engine/Plugins/Online/Experimental/OnlineServicesEOS/Source/Private/Online/AuthEOS.cpp b/Engine/Plugins/Online/Experimental/OnlineServicesEOS/Source/Private/Online/AuthEOS.cpp
--- a/Engine/Plugins/Online/Experimental/OnlineServicesEOS/Source/Private/Online/AuthEOS.cpp
+++ b/Engine/Plugins/Online/Experimental/OnlineServicesEOS/Source/Private/Online/AuthEOS.cpp
@@ -216,7 +216,7 @@ TOnlineAsyncOpHandle<FAuthLogin> FAuthEOS::Login(FAuthLogin::Params&& Params)
 		}
 		else
 		{
-			FOnlineError CopyUserInfoError(FromEOSError(CopyUserInfoResult));
+			FOnlineError CopyUserInfoError(Errors::FromEOSResult(CopyUserInfoResult));
 			UE_LOG(LogOnlineServices, Warning, TEXT("[FAuthEOS::Login] Failure: EOS_UserInfo_CopyUserInfo %s"), *CopyUserInfoError.GetLogString());
 			InAsyncOp.SetError(Errors::Unknown(MoveTemp(CopyUserInfoError)));
 
@@ -308,7 +308,7 @@ TOnlineAsyncOpHandle<FAuthQueryExternalServerAuthTicket> FAuthEOS::QueryExternal
 		}
 		else
 		{
-			InAsyncOp.SetError(Errors::Unknown(FromEOSError(Result)));
+			InAsyncOp.SetError(Errors::FromEOSResult(Result));
 		}
 	})
 	.Enqueue(GetSerialQueue());
diff --git a/Engine/Plugins/Online/Experimental/OnlineServicesEOS/Source/Private/Online/PresenceEOS.cpp b/Engine/Plugins/Online/Experimental/OnlineServicesEOS/Source/Private/Online/PresenceEOS.cpp
--- a/Engine/Plugins/Online/Experimental/OnlineServicesEOS/Source/Private/Online/PresenceEOS.cpp
+++ b/Engine/Plugins/Online/Experimental/OnlineServicesEOS/Source/Private/Online/PresenceEOS.cpp
@@ -6,7 +6,7 @@
 #include "Online/OnlineServicesEOS.h"
 #include "Online/OnlineServicesEOSTypes.h"
 #include "Online/AuthEOS.h"
-#include "Online/ErrorsEOS.h"
+#include "Online/OnlineErrorEOSGS.h"
 
 #include "eos_presence.h"
 
@@ -144,7 +144,7 @@ TOnlineAsyncOpHandle<FQueryPresence> FPresenceEOS::QueryPresence(FQueryPresence:
 				}
 				else
 				{
-					InAsyncOp.SetError(Errors::FromEOSResult(Data->ResultCode)); 
+					InAsyncOp.SetError(Errors::FromEOSResult(Data->ResultCode));
 				}
 			})
 			.Enqueue(GetSerialQueue());
@@ -539,7 +539,7 @@ TOnlineAsyncOpHandle<FPartialUpdatePresence> FPresenceEOS::PartialUpdatePresence
 				}
 				else
 				{
-					InAsyncOp.SetError(Errors::FromEOSResult(Data->ResultCode)); 
+					InAsyncOp.SetError(Errors::FromEOSResult(Data->ResultCode));
 				}
 			})
 			.Enqueue(GetSerialQueue());
diff --git a/Engine/Plugins/Online/Experimental/OnlineServicesEOS/Source/Private/Online/UserInfoEOS.cpp b/Engine/Plugins/Online/Experimental/OnlineServicesEOS/Source/Private/Online/UserInfoEOS.cpp
--- a/Engine/Plugins/Online/Experimental/OnlineServicesEOS/Source/Private/Online/UserInfoEOS.cpp
+++ b/Engine/Plugins/Online/Experimental/OnlineServicesEOS/Source/Private/Online/UserInfoEOS.cpp
@@ -6,7 +6,7 @@
 #include "Online/OnlineServicesEOS.h"
 #include "Online/OnlineServicesEOSTypes.h"
 #include "Online/AuthEOS.h"
-#include "Online/ErrorsEOS.h"
+#include "Online/OnlineErrorEOSGS.h"
 
 #include "eos_userinfo.h"
 
@@ -69,7 +69,7 @@ TOnlineAsyncOpHandle<FQueryUserInfo> FUserInfoEOS::QueryUserInfo(FQueryUserInfo:
 			if (CallbackInfo->ResultCode != EOS_EResult::EOS_Success)
 			{
 				UE_LOG(LogTemp, Warning, TEXT("EOS_UserInfo_QueryUserInfo failed with result=[%s]"), *LexToString(CallbackInfo->ResultCode));
-				Op.SetError(FromEOSError(CallbackInfo->ResultCode));
+				Op.SetError(Errors::FromEOSResult(CallbackInfo->ResultCode));
 			}
 		});
 	}
@@ -112,7 +112,7 @@ TOnlineResult<FGetUserInfo> FUserInfoEOS::GetUserInfo(FGetUserInfo::Params&& Par
 	if(EosResult != EOS_EResult::EOS_Success)
 	{
 		UE_LOG(LogTemp, Warning, TEXT("EOS_UserInfo_CopyUserInfo failed with result=[%s]"), *LexToString(EosResult));
-		return TOnlineResult<FGetUserInfo>(FromEOSError(EosResult));
+		return TOnlineResult<FGetUserInfo>(Errors::FromEOSResult(EosResult));
 	}
 
 	TSharedRef<FUserInfo> UserInfo = MakeShared<FUserInfo>();
diff --git a/Engine/Plugins/Online/Experimental/OnlineServicesEOS/Source/Public/Online/ErrorsEOS.h b/Engine/Plugins/Online/Experimental/OnlineServicesEOS/Source/Public/Online/ErrorsEOS.h
deleted file mode 100644
diff --git a/Engine/Plugins/Online/Experimental/OnlineServicesEOSGS/Source/Private/Online/AchievementsEOSGS.cpp b/Engine/Plugins/Online/Experimental/OnlineServicesEOSGS/Source/Private/Online/AchievementsEOSGS.cpp
--- a/Engine/Plugins/Online/Experimental/OnlineServicesEOSGS/Source/Private/Online/AchievementsEOSGS.cpp
+++ b/Engine/Plugins/Online/Experimental/OnlineServicesEOSGS/Source/Private/Online/AchievementsEOSGS.cpp
@@ -83,7 +83,7 @@ TOnlineAsyncOpHandle<FQueryAchievementDefinitions> FAchievementsEOSGS::QueryAchi
 			{
 				UE_LOG(LogTemp, Warning, TEXT("EOS_Achievements_QueryDefinitions failed with result=[%s]"), *LexToString(Data->ResultCode));
 
-				InAsyncOp.SetError(FromEOSError(Data->ResultCode));
+				InAsyncOp.SetError(Errors::FromEOSResult(Data->ResultCode));
 				return;
 			}
 
@@ -107,7 +107,7 @@ TOnlineAsyncOpHandle<FQueryAchievementDefinitions> FAchievementsEOSGS::QueryAchi
 				{
 					UE_LOG(LogTemp, Warning, TEXT("EOS_Achievements_CopyAchievementDefinitionV2ByIndex failed with result=[%s]"), *LexToString(CopyResult));
 
-					InAsyncOp.SetError(FromEOSError(CopyResult));
+					InAsyncOp.SetError(Errors::FromEOSResult(CopyResult));
 					return;
 				}
 
@@ -228,7 +228,7 @@ TOnlineAsyncOpHandle<FQueryAchievementStates> FAchievementsEOSGS::QueryAchieveme
 			{
 				UE_LOG(LogTemp, Warning, TEXT("EOS_Achievements_QueryPlayerAchievements failed with result=[%s]"), *LexToString(Data->ResultCode));
 
-				InAsyncOp.SetError(FromEOSError(Data->ResultCode));
+				InAsyncOp.SetError(Errors::FromEOSResult(Data->ResultCode));
 				return;
 			}
 
@@ -255,7 +255,7 @@ TOnlineAsyncOpHandle<FQueryAchievementStates> FAchievementsEOSGS::QueryAchieveme
 				{
 					UE_LOG(LogTemp, Warning, TEXT("EOS_Achievements_CopyPlayerAchievementByIndex failed with result=[%s]"), *LexToString(CopyResult));
 
-					InAsyncOp.SetError(FromEOSError(CopyResult));
+					InAsyncOp.SetError(Errors::FromEOSResult(CopyResult));
 					return;
 				}
 				
@@ -377,7 +377,7 @@ TOnlineAsyncOpHandle<FUnlockAchievements> FAchievementsEOSGS::UnlockAchievements
 		{
 			UE_LOG(LogTemp, Warning, TEXT("EOS_Achievements_UnlockAchievements failed with result=[%s]"), *LexToString(Data->ResultCode));
 
-			InAsyncOp.SetError(FromEOSError(Data->ResultCode));
+			InAsyncOp.SetError(Errors::FromEOSResult(Data->ResultCode));
 		}
 		else
 		{
diff --git a/Engine/Plugins/Online/Experimental/OnlineServicesEOSGS/Source/Private/Online/AuthEOSGS.cpp b/Engine/Plugins/Online/Experimental/OnlineServicesEOSGS/Source/Private/Online/AuthEOSGS.cpp
--- a/Engine/Plugins/Online/Experimental/OnlineServicesEOSGS/Source/Private/Online/AuthEOSGS.cpp
+++ b/Engine/Plugins/Online/Experimental/OnlineServicesEOSGS/Source/Private/Online/AuthEOSGS.cpp
@@ -762,7 +762,7 @@ TOnlineAsyncOpHandle<FAuthLogin> FAuthEOSGS::Login(FAuthLogin::Params&& Params)
 			}
 			else
 			{
-				FOnlineError ProductUserIdError(FromEOSError(ProductUserIdResult));
+				FOnlineError ProductUserIdError(Errors::FromEOSResult(ProductUserIdResult));
 				UE_LOG(LogOnlineServices, Warning, TEXT("[FAuthEOSGS::Login] Failure: EOS_ProductUserId_ToString %s"), *ProductUserIdError.GetLogString());
 				InAsyncOp.SetError(Errors::Unknown(MoveTemp(ProductUserIdError)));
 
@@ -853,7 +853,7 @@ TOnlineAsyncOpHandle<FAuthLogout> FAuthEOSGS::Logout(FAuthLogout::Params&& Param
 			{
 				if (Data->ResultCode != EOS_EResult::EOS_Success)
 				{
-					FOnlineError DeletePersistentAuthError(FromEOSError(Data->ResultCode));
+					FOnlineError DeletePersistentAuthError(Errors::FromEOSResult(Data->ResultCode));
 					UE_LOG(LogOnlineServices, Warning, TEXT("[FAuthEOSGS::Logout] Failure: DeletePersistentAuthResult %s"), *DeletePersistentAuthError.GetLogString());
 				}
 
@@ -1021,13 +1021,13 @@ TFuture<TDefaultErrorResult<FAuthLoginEASImpl>> FAuthEOSGS::LoginEASImpl(const F
 				}
 				else
 				{
-					Promise.SetValue(TDefaultErrorResult<FAuthLoginEASImpl>(FromEOSError(Data->ResultCode)));
+					Promise.SetValue(TDefaultErrorResult<FAuthLoginEASImpl>(Errors::FromEOSResult(Data->ResultCode)));
 				}
 			});
 		}
 		else
 		{
-			FOnlineError ResolvedError = (Data->ResultCode == EOS_EResult::EOS_InvalidAuth) ? Errors::InvalidCreds() : Errors::Unknown(FromEOSError(Data->ResultCode));
+			FOnlineError ResolvedError = Errors::FromEOSResult(Data->ResultCode);
 
 			const bool bShouldRemoveCachedToken =
 				Data->ResultCode == EOS_EResult::EOS_InvalidAuth ||
@@ -1045,7 +1045,7 @@ TFuture<TDefaultErrorResult<FAuthLoginEASImpl>> FAuthEOSGS::LoginEASImpl(const F
 				{
 					if (Data->ResultCode != EOS_EResult::EOS_Success)
 					{
-						FOnlineError DeletePersistentAuthError(FromEOSError(Data->ResultCode));
+						FOnlineError DeletePersistentAuthError(Errors::FromEOSResult(Data->ResultCode));
 						UE_LOG(LogOnlineServices, Warning, TEXT("[FAuthEOS::LoginEASImpl] Failure: DeletePersistentAuthResult %s"), *DeletePersistentAuthError.GetLogString());
 					}
 
@@ -1082,7 +1082,7 @@ TFuture<TDefaultErrorResult<FAuthLogoutEASImpl>> FAuthEOSGS::LogoutEASImpl(const
 		}
 		else
 		{
-			Promise.SetValue(TDefaultErrorResult<FAuthLogoutEASImpl>(FromEOSError(Data->ResultCode)));
+			Promise.SetValue(TDefaultErrorResult<FAuthLogoutEASImpl>(Errors::FromEOSResult(Data->ResultCode)));
 		}
 	});
 
@@ -1111,7 +1111,7 @@ TDefaultErrorResult<FAuthGetExternalAuthTokenImpl> FAuthEOSGS::GetExternalAuthTo
 	}
 	else
 	{
-		return TDefaultErrorResult<FAuthGetExternalAuthTokenImpl>(FromEOSError(Result));
+		return TDefaultErrorResult<FAuthGetExternalAuthTokenImpl>(Errors::FromEOSResult(Result));
 	}
 }
 
@@ -1152,13 +1152,13 @@ TFuture<TDefaultErrorResult<FAuthLoginConnectImpl>> FAuthEOSGS::LoginConnectImpl
 				}
 				else
 				{
-					Promise.SetValue(TDefaultErrorResult<FAuthLoginConnectImpl>(FromEOSError(Data->ResultCode)));
+					Promise.SetValue(TDefaultErrorResult<FAuthLoginConnectImpl>(Errors::FromEOSResult(Data->ResultCode)));
 				}
 			});
 		}
 		else
 		{
-			Promise.SetValue(TDefaultErrorResult<FAuthLoginConnectImpl>(FromEOSError(Data->ResultCode)));
+			Promise.SetValue(TDefaultErrorResult<FAuthLoginConnectImpl>(Errors::FromEOSResult(Data->ResultCode)));
 		}
 	});
 
diff --git a/Engine/Plugins/Online/Experimental/OnlineServicesEOSGS/Source/Private/Online/LobbiesEOSGS.cpp b/Engine/Plugins/Online/Experimental/OnlineServicesEOSGS/Source/Private/Online/LobbiesEOSGS.cpp
--- a/Engine/Plugins/Online/Experimental/OnlineServicesEOSGS/Source/Private/Online/LobbiesEOSGS.cpp
+++ b/Engine/Plugins/Online/Experimental/OnlineServicesEOSGS/Source/Private/Online/LobbiesEOSGS.cpp
@@ -129,7 +129,7 @@ TOnlineAsyncOpHandle<FCreateLobby> FLobbiesEOSGS::CreateLobby(FCreateLobby::Para
 		if (Data->ResultCode != EOS_EResult::EOS_Success)
 		{
 			// Todo: errors
-			InAsyncOp.SetError(Errors::Unknown(FromEOSError(Data->ResultCode)));
+			InAsyncOp.SetError(Errors::FromEOSResult(Data->ResultCode));
 			return MakeFulfilledPromise<TDefaultErrorResultInternal<TSharedRef<FLobbyDataEOS>>>().GetFuture();
 		}
 
@@ -1345,7 +1345,7 @@ TOnlineAsyncOpHandle<FLobbiesJoinLobbyMemberImpl> FLobbiesEOSGS::JoinLobbyMember
 		if (Data->ResultCode != EOS_EResult::EOS_Success)
 		{
 			// TODO: Error codes
-			InAsyncOp.SetError(Errors::Unknown(FromEOSError(Data->ResultCode)));
+			InAsyncOp.SetError(Errors::FromEOSResult(Data->ResultCode));
 		}
 	})
 	.Then([this](TOnlineAsyncOp<FLobbiesJoinLobbyMemberImpl>& InAsyncOp)
@@ -1428,7 +1428,7 @@ TFuture<TDefaultErrorResult<FLobbiesLeaveLobbyImpl>> FLobbiesEOSGS::LeaveLobbyIm
 		if (Result->ResultCode != EOS_EResult::EOS_Success)
 		{
 			// Todo: Errors
-			Promise.EmplaceValue(Errors::Unknown(FromEOSError(Result->ResultCode)));
+			Promise.EmplaceValue(Errors::FromEOSResult(Result->ResultCode));
 			return;
 		}
 
@@ -1465,7 +1465,7 @@ TFuture<TDefaultErrorResult<FLobbiesDestroyLobbyImpl>> FLobbiesEOSGS::DestroyLob
 		if (Result->ResultCode != EOS_EResult::EOS_Success)
 		{
 			// Todo: Errors
-			Promise.EmplaceValue(Errors::Unknown(FromEOSError(Result->ResultCode)));
+			Promise.EmplaceValue(Errors::FromEOSResult(Result->ResultCode));
 			return;
 		}
 
@@ -1503,7 +1503,7 @@ TFuture<TDefaultErrorResult<FLobbiesInviteLobbyMemberImpl>> FLobbiesEOSGS::Invit
 		if (Result->ResultCode != EOS_EResult::EOS_Success)
 		{
 			// Todo: Errors
-			Promise.EmplaceValue(Errors::Unknown(FromEOSError(Result->ResultCode)));
+			Promise.EmplaceValue(Errors::FromEOSResult(Result->ResultCode));
 			return;
 		}
 
@@ -1546,7 +1546,7 @@ TFuture<TDefaultErrorResult<FLobbiesDeclineLobbyInvitationImpl>> FLobbiesEOSGS::
 		if (Result->ResultCode != EOS_EResult::EOS_Success)
 		{
 			// Todo: Errors
-			Promise.EmplaceValue(Errors::Unknown(FromEOSError(Result->ResultCode)));
+			Promise.EmplaceValue(Errors::FromEOSResult(Result->ResultCode));
 			return;
 		}
 
@@ -1586,7 +1586,7 @@ TFuture<TDefaultErrorResult<FLobbiesKickLobbyMemberImpl>> FLobbiesEOSGS::KickLob
 		if (Result->ResultCode != EOS_EResult::EOS_Success)
 		{
 			// Todo: Errors
-			Promise.EmplaceValue(Errors::Unknown(FromEOSError(Result->ResultCode)));
+			Promise.EmplaceValue(Errors::FromEOSResult(Result->ResultCode));
 			return;
 		}
 
@@ -1626,7 +1626,7 @@ TFuture<TDefaultErrorResult<FLobbiesPromoteLobbyMemberImpl>> FLobbiesEOSGS::Prom
 		if (Result->ResultCode != EOS_EResult::EOS_Success)
 		{
 			// Todo: Errors
-			Promise.EmplaceValue(Errors::Unknown(FromEOSError(Result->ResultCode)));
+			Promise.EmplaceValue(Errors::FromEOSResult(Result->ResultCode));
 			return;
 		}
 
@@ -1675,7 +1675,7 @@ TFuture<TDefaultErrorResult<FLobbiesModifyLobbyDataImpl>> FLobbiesEOSGS::ModifyL
 		if (Future.Get() != EOS_EResult::EOS_Success && Future.Get() != EOS_EResult::EOS_NoChange)
 		{
 			// Todo: Errors
-			Promise.EmplaceValue(Errors::Unknown(FromEOSError(Future.Get())));
+			Promise.EmplaceValue(Errors::FromEOSResult(Future.Get()));
 			return;
 		}
 
@@ -1714,7 +1714,7 @@ TFuture<TDefaultErrorResult<FLobbiesModifyLobbyMemberDataImpl>> FLobbiesEOSGS::M
 		if (Future.Get() != EOS_EResult::EOS_Success && Future.Get() != EOS_EResult::EOS_NoChange)
 		{
 			// Todo: Errors
-			Promise.EmplaceValue(Errors::Unknown(FromEOSError(Future.Get())));
+			Promise.EmplaceValue(Errors::FromEOSResult(Future.Get()));
 			return;
 		}
 
diff --git a/Engine/Plugins/Online/Experimental/OnlineServicesEOSGS/Source/Private/Online/LobbiesEOSGSTypes.cpp b/Engine/Plugins/Online/Experimental/OnlineServicesEOSGS/Source/Private/Online/LobbiesEOSGSTypes.cpp
--- a/Engine/Plugins/Online/Experimental/OnlineServicesEOSGS/Source/Private/Online/LobbiesEOSGSTypes.cpp
+++ b/Engine/Plugins/Online/Experimental/OnlineServicesEOSGS/Source/Private/Online/LobbiesEOSGSTypes.cpp
@@ -159,7 +159,7 @@ TDefaultErrorResultInternal<TSharedRef<FLobbyDetailsEOS>> FLobbyDetailsEOS::Crea
 	EOS_EResult EOSResult = EOS_Lobby_CopyLobbyDetailsHandle(Prerequisites->LobbyInterfaceHandle, &Options, &LobbyDetailsHandle);
 	if (EOSResult != EOS_EResult::EOS_Success)
 	{
-		return TDefaultErrorResultInternal<TSharedRef<FLobbyDetailsEOS>>(FromEOSError(EOSResult));
+		return TDefaultErrorResultInternal<TSharedRef<FLobbyDetailsEOS>>(Errors::FromEOSResult(EOSResult));
 	}
 
 	TDefaultErrorResultInternal<TSharedRef<FLobbyDetailsInfoEOS>> LobbyDetailsInfoResult = FLobbyDetailsInfoEOS::Create(LobbyDetailsHandle);
@@ -186,7 +186,7 @@ TDefaultErrorResultInternal<TSharedRef<FLobbyDetailsEOS>> FLobbyDetailsEOS::Crea
 	EOS_EResult EOSResult = EOS_Lobby_CopyLobbyDetailsHandleByInviteId(Prerequisites->LobbyInterfaceHandle, &Options, &LobbyDetailsHandle);
 	if (EOSResult != EOS_EResult::EOS_Success)
 	{
-		return TDefaultErrorResultInternal<TSharedRef<FLobbyDetailsEOS>>(FromEOSError(EOSResult));
+		return TDefaultErrorResultInternal<TSharedRef<FLobbyDetailsEOS>>(Errors::FromEOSResult(EOSResult));
 	}
 
 	TDefaultErrorResultInternal<TSharedRef<FLobbyDetailsInfoEOS>> LobbyDetailsInfoResult = FLobbyDetailsInfoEOS::Create(LobbyDetailsHandle);
@@ -213,7 +213,7 @@ TDefaultErrorResultInternal<TSharedRef<FLobbyDetailsEOS>> FLobbyDetailsEOS::Crea
 	EOS_EResult EOSResult = EOS_Lobby_CopyLobbyDetailsHandleByUiEventId(Prerequisites->LobbyInterfaceHandle, &Options, &LobbyDetailsHandle);
 	if (EOSResult != EOS_EResult::EOS_Success)
 	{
-		return TDefaultErrorResultInternal<TSharedRef<FLobbyDetailsEOS>>(FromEOSError(EOSResult));
+		return TDefaultErrorResultInternal<TSharedRef<FLobbyDetailsEOS>>(Errors::FromEOSResult(EOSResult));
 	}
 
 	TDefaultErrorResultInternal<TSharedRef<FLobbyDetailsInfoEOS>> LobbyDetailsInfoResult = FLobbyDetailsInfoEOS::Create(LobbyDetailsHandle);
@@ -241,7 +241,7 @@ TDefaultErrorResultInternal<TSharedRef<FLobbyDetailsEOS>> FLobbyDetailsEOS::Crea
 	EOS_EResult EOSResult = EOS_LobbySearch_CopySearchResultByIndex(SearchHandle, &Options, &LobbyDetailsHandle);
 	if (EOSResult != EOS_EResult::EOS_Success)
 	{
-		return TDefaultErrorResultInternal<TSharedRef<FLobbyDetailsEOS>>(FromEOSError(EOSResult));
+		return TDefaultErrorResultInternal<TSharedRef<FLobbyDetailsEOS>>(Errors::FromEOSResult(EOSResult));
 	}
 
 	TDefaultErrorResultInternal<TSharedRef<FLobbyDetailsInfoEOS>> LobbyDetailsInfoResult = FLobbyDetailsInfoEOS::Create(LobbyDetailsHandle);
@@ -343,7 +343,7 @@ TFuture<TDefaultErrorResultInternal<TSharedRef<FClientLobbySnapshot>>> FLobbyDet
 				if (EOSResult != EOS_EResult::EOS_Success)
 				{
 					// todo: errors
-					Promise.EmplaceValue(FromEOSError(EOSResult));
+					Promise.EmplaceValue(Errors::FromEOSResult(EOSResult));
 					return;
 				}
 
@@ -391,7 +391,7 @@ TDefaultErrorResultInternal<TSharedRef<FClientLobbyMemberSnapshot>> FLobbyDetail
 			EOS_EResult EOSResult = EOS_LobbyDetails_CopyMemberAttributeByIndex(LobbyDetailsHandle, &CopyMemberAttributeByIndexOptions, &LobbyAttribute);
 			if (EOSResult != EOS_EResult::EOS_Success)
 			{
-				return TDefaultErrorResultInternal<TSharedRef<FClientLobbyMemberSnapshot>>(FromEOSError(EOSResult));
+				return TDefaultErrorResultInternal<TSharedRef<FClientLobbyMemberSnapshot>>(Errors::FromEOSResult(EOSResult));
 			}
 
 			FLobbyAttributeTranslator<ELobbyTranslationType::FromService> AttributeTranslator(*LobbyAttribute->Data);
@@ -587,7 +587,7 @@ TDefaultErrorResultInternal<TSharedRef<FLobbyDetailsInfoEOS>> FLobbyDetailsInfoE
 	EOS_EResult EOSResult = EOS_LobbyDetails_CopyInfo(LobbyDetailsHandle, &CopyInfoOptions, &LobbyDetailsInfo);
 	if (EOSResult != EOS_EResult::EOS_Success)
 	{
-		return TDefaultErrorResultInternal<TSharedRef<FLobbyDetailsInfoEOS>>(FromEOSError(EOSResult));
+		return TDefaultErrorResultInternal<TSharedRef<FLobbyDetailsInfoEOS>>(Errors::FromEOSResult(EOSResult));
 	}
 
 	return TDefaultErrorResultInternal<TSharedRef<FLobbyDetailsInfoEOS>>(MakeShared<FLobbyDetailsInfoEOS>(FLobbyDetailsInfoPtr(LobbyDetailsInfo)));
@@ -859,7 +859,7 @@ TFuture<TDefaultErrorResultInternal<TSharedRef<FLobbySearchEOS>>> FLobbySearchEO
 	if (EOSResult != EOS_EResult::EOS_Success)
 	{
 		// todo: errors
-		return MakeFulfilledPromise<TDefaultErrorResultInternal<TSharedRef<FLobbySearchEOS>>>(FromEOSError(EOSResult)).GetFuture();
+		return MakeFulfilledPromise<TDefaultErrorResultInternal<TSharedRef<FLobbySearchEOS>>>(Errors::FromEOSResult(EOSResult)).GetFuture();
 	}
 
 	if (Params.LobbyId) // Search for specific lobby.
@@ -878,7 +878,7 @@ TFuture<TDefaultErrorResultInternal<TSharedRef<FLobbySearchEOS>>> FLobbySearchEO
 		if (EOSResult != EOS_EResult::EOS_Success)
 		{
 			// todo: errors
-			return MakeFulfilledPromise<TDefaultErrorResultInternal<TSharedRef<FLobbySearchEOS>>>(FromEOSError(EOSResult)).GetFuture();
+			return MakeFulfilledPromise<TDefaultErrorResultInternal<TSharedRef<FLobbySearchEOS>>>(Errors::FromEOSResult(EOSResult)).GetFuture();
 		}
 	}
 	else if (Params.TargetUser) // Search for specific user.
@@ -891,7 +891,7 @@ TFuture<TDefaultErrorResultInternal<TSharedRef<FLobbySearchEOS>>> FLobbySearchEO
 		if (EOSResult != EOS_EResult::EOS_Success)
 		{
 			// todo: errors
-			return MakeFulfilledPromise<TDefaultErrorResultInternal<TSharedRef<FLobbySearchEOS>>>(FromEOSError(EOSResult)).GetFuture();
+			return MakeFulfilledPromise<TDefaultErrorResultInternal<TSharedRef<FLobbySearchEOS>>>(Errors::FromEOSResult(EOSResult)).GetFuture();
 		}
 	}
 	else // Search using parameters.
@@ -915,7 +915,7 @@ TFuture<TDefaultErrorResultInternal<TSharedRef<FLobbySearchEOS>>> FLobbySearchEO
 			if (EOSResult != EOS_EResult::EOS_Success)
 			{
 				// todo: errors
-				return MakeFulfilledPromise<TDefaultErrorResultInternal<TSharedRef<FLobbySearchEOS>>>(FromEOSError(EOSResult)).GetFuture();
+				return MakeFulfilledPromise<TDefaultErrorResultInternal<TSharedRef<FLobbySearchEOS>>>(Errors::FromEOSResult(EOSResult)).GetFuture();
 			}
 		}
 
@@ -933,7 +933,7 @@ TFuture<TDefaultErrorResultInternal<TSharedRef<FLobbySearchEOS>>> FLobbySearchEO
 			if (EOSResult != EOS_EResult::EOS_Success)
 			{
 				// todo: errors
-				return MakeFulfilledPromise<TDefaultErrorResultInternal<TSharedRef<FLobbySearchEOS>>>(FromEOSError(EOSResult)).GetFuture();
+				return MakeFulfilledPromise<TDefaultErrorResultInternal<TSharedRef<FLobbySearchEOS>>>(Errors::FromEOSResult(EOSResult)).GetFuture();
 			}
 		}
 	}
@@ -952,7 +952,7 @@ TFuture<TDefaultErrorResultInternal<TSharedRef<FLobbySearchEOS>>> FLobbySearchEO
 		if (CallbackInfo->ResultCode != EOS_EResult::EOS_Success)
 		{
 			// todo: errors
-			Promise.EmplaceValue(FromEOSError(CallbackInfo->ResultCode));
+			Promise.EmplaceValue(Errors::FromEOSResult(CallbackInfo->ResultCode));
 			return;
 		}
 
diff --git a/Engine/Plugins/Online/Experimental/OnlineServicesEOSGS/Source/Private/Online/OnlineErrorEOSGS.cpp b/Engine/Plugins/Online/Experimental/OnlineServicesEOSGS/Source/Private/Online/OnlineErrorEOSGS.cpp
--- a/Engine/Plugins/Online/Experimental/OnlineServicesEOSGS/Source/Private/Online/OnlineErrorEOSGS.cpp
+++ b/Engine/Plugins/Online/Experimental/OnlineServicesEOSGS/Source/Private/Online/OnlineErrorEOSGS.cpp
@@ -1,8 +1,18 @@
 // Copyright Epic Games, Inc. All Rights Reserved.
 
-#include "OnlineErrorEOSGS.h"
+#include "Online/OnlineErrorEOSGS.h"
+#include "Online/OnlineErrorDefinitions.h"
 
-namespace UE::Online {
+namespace UE::Online::Errors {
+
+class FOnlineErrorDetailsEOS : public IOnlineErrorDetails
+{
+public:
+	static const TSharedRef<const IOnlineErrorDetails, ESPMode::ThreadSafe>& Get();
+
+	virtual FText GetText(const FOnlineError& OnlineError) const override;
+	virtual FString GetLogString(const FOnlineError& OnlineError) const override;
+};
 
 const TSharedRef<const IOnlineErrorDetails, ESPMode::ThreadSafe>& FOnlineErrorDetailsEOS::Get()
 {
@@ -12,29 +22,52 @@ const TSharedRef<const IOnlineErrorDetails, ESPMode::ThreadSafe>& FOnlineErrorDe
 
 FText FOnlineErrorDetailsEOS::GetText(const FOnlineError& OnlineError) const
 {
-	const EOS_EResult EosResult = EOS_EResult(OnlineError.GetValue());
-	return FText::FromStringView(FStringView(UTF8_TO_TCHAR(EOS_EResult_ToString(EosResult))));
+	return FText::FromString(LexToString(OnlineError.GetValue()));
 }
 
-FString FOnlineErrorDetailsEOS::GetLogString(const FOnlineError& OnlineError) const 
+FString FOnlineErrorDetailsEOS::GetLogString(const FOnlineError& OnlineError) const
 {
-	const EOS_EResult EosResult = EOS_EResult(OnlineError.GetValue());
-	return UTF8_TO_TCHAR(EOS_EResult_ToString(EosResult));
+	return LexToString(OnlineError.GetValue());
 }
 
-bool operator==(const UE::Online::FOnlineError& OnlineError, EOS_EResult EosResult)
+/** This doesn't wrap every single error in EOS, only the common ones that are strongly related to common errors */
+FOnlineError MapCommonEOSError(FOnlineError&& Error, EOS_EResult Result)
 {
-	return OnlineError.GetErrorCode() == CreateErrorCode(EosResult);
+	switch (Result)
+	{
+	case EOS_EResult::EOS_Success:				return Errors::Success(MoveTemp(Error));
+	case EOS_EResult::EOS_NoConnection:			return Errors::NoConnection(MoveTemp(Error));
+	case EOS_EResult::EOS_InvalidCredentials:	return Errors::InvalidCreds(MoveTemp(Error));
+	case EOS_EResult::EOS_InvalidUser:			return Errors::InvalidUser(MoveTemp(Error));
+	case EOS_EResult::EOS_InvalidAuth:			return Errors::InvalidAuth(MoveTemp(Error));
+	case EOS_EResult::EOS_AccessDenied:			return Errors::AccessDenied(MoveTemp(Error));
+	case EOS_EResult::EOS_MissingPermissions:	return Errors::AccessDenied(MoveTemp(Error));
+	case EOS_EResult::EOS_TooManyRequests:		return Errors::TooManyRequests(MoveTemp(Error));
+	case EOS_EResult::EOS_AlreadyPending:		return Errors::AlreadyPending(MoveTemp(Error));
+	case EOS_EResult::EOS_InvalidParameters:	return Errors::InvalidParams(MoveTemp(Error));
+	case EOS_EResult::EOS_InvalidRequest:		return Errors::InvalidParams(MoveTemp(Error));
+	case EOS_EResult::EOS_UnrecognizedResponse: return Errors::InvalidResults(MoveTemp(Error));
+	case EOS_EResult::EOS_IncompatibleVersion:	return Errors::IncompatibleVersion(MoveTemp(Error));
+	case EOS_EResult::EOS_NotConfigured:		return Errors::NotConfigured(MoveTemp(Error));
+	case EOS_EResult::EOS_NotImplemented:		return Errors::NotImplemented(MoveTemp(Error));
+	case EOS_EResult::EOS_Canceled:				return Errors::Cancelled(MoveTemp(Error));
+	case EOS_EResult::EOS_NotFound:				return Errors::NotFound(MoveTemp(Error));
+	case EOS_EResult::EOS_OperationWillRetry:	return Errors::WillRetry(MoveTemp(Error));
+	case EOS_EResult::EOS_VersionMismatch:		return Errors::IncompatibleVersion(MoveTemp(Error));
+	case EOS_EResult::EOS_LimitExceeded:		return Errors::TooManyRequests(MoveTemp(Error));
+	case EOS_EResult::EOS_TimedOut:				return Errors::Timeout(MoveTemp(Error));
+	default:									return Errors::Unknown(MoveTemp(Error));
+	}
 }
 
-ErrorCodeType CreateErrorCode(EOS_EResult EosResult)
+ErrorCodeType ErrorCodeFromEOSResult(EOS_EResult EosResult)
 {
 	return Errors::ErrorCode::Create(Errors::ErrorCode::Category::EOS_System, Errors::ErrorCode::Category::EOS, uint32(EosResult));
 }
 
-FOnlineError CreateOnlineError(EOS_EResult EosResult)
+FOnlineError FromEOSResult(EOS_EResult EosResult, FErrorMapperEosFn&& MapperFn)
 {
-	return FOnlineError(CreateErrorCode(EosResult), FOnlineErrorDetailsEOS::Get());
+	return MapperFn(FOnlineError(ErrorCodeFromEOSResult(EosResult), FOnlineErrorDetailsEOS::Get()), EosResult);
 }
 
-} /* namespace UE::Online */
+} /* namespace UE::Online::Errors */
diff --git a/Engine/Plugins/Online/Experimental/OnlineServicesEOSGS/Source/Private/Online/OnlineErrorEOSGS.h b/Engine/Plugins/Online/Experimental/OnlineServicesEOSGS/Source/Private/Online/OnlineErrorEOSGS.h
deleted file mode 100644
diff --git a/Engine/Plugins/Online/Experimental/OnlineServicesEOSGS/Source/Private/Online/SessionsEOSGS.cpp b/Engine/Plugins/Online/Experimental/OnlineServicesEOSGS/Source/Private/Online/SessionsEOSGS.cpp
--- a/Engine/Plugins/Online/Experimental/OnlineServicesEOSGS/Source/Private/Online/SessionsEOSGS.cpp
+++ b/Engine/Plugins/Online/Experimental/OnlineServicesEOSGS/Source/Private/Online/SessionsEOSGS.cpp
@@ -282,7 +282,7 @@ TOnlineAsyncOpHandle<FCreateSession> FSessionsEOSGS::CreateSession(FCreateSessio
 		{
 			UE_LOG(LogTemp, Error, TEXT("[FSessionsEOSGS::CreateSession] EOS_Sessions_CreateSessionModification failed with result [%s]"), *LexToString(ResultCode));
 
-			Op.SetError(FromEOSError(ResultCode));
+			Op.SetError(Errors::FromEOSResult(ResultCode));
 
 			return MakeFulfilledPromise<TDefaultErrorResult<FUpdateSessionImpl>>().GetFuture();
 		}
@@ -580,7 +580,7 @@ TOnlineAsyncOpHandle<FUpdateSession> FSessionsEOSGS::UpdateSession(FUpdateSessio
 		{
 			UE_LOG(LogTemp, Error, TEXT("[FSessionsEOSGS::UpdateSession] EOS_Sessions_UpdateSessionModification failed with result [%s]"), *LexToString(ResultCode));
 
-			Op.SetError(FromEOSError(ResultCode));
+			Op.SetError(Errors::FromEOSResult(ResultCode));
 
 			return MakeFulfilledPromise<TDefaultErrorResult<FUpdateSessionImpl>>().GetFuture();
 		}
@@ -640,7 +640,7 @@ TFuture<TDefaultErrorResult<FUpdateSessionImpl>> FSessionsEOSGS::UpdateSessionIm
 	{
 		if (Result->ResultCode != EOS_EResult::EOS_Success && Result->ResultCode != EOS_EResult::EOS_Sessions_OutOfSync)
 		{
-			Promise.EmplaceValue(FromEOSError(Result->ResultCode));
+			Promise.EmplaceValue(Errors::FromEOSResult(Result->ResultCode));
 			return;
 		}
 
@@ -697,7 +697,7 @@ TOnlineAsyncOpHandle<FLeaveSession> FSessionsEOSGS::LeaveSession(FLeaveSession::
 	{
 		if (Result->ResultCode != EOS_EResult::EOS_Success)
 		{
-			Op.SetError(FromEOSError(Result->ResultCode));
+			Op.SetError(Errors::FromEOSResult(Result->ResultCode));
 			return;
 		}
 
@@ -846,7 +846,7 @@ TOnlineAsyncOpHandle<FFindSessions> FSessionsEOSGS::FindSessions(FFindSessions::
 		{
 			UE_LOG(LogTemp, Error, TEXT("[FSessionsEOSGS::FindSessions] EOS_Sessions_CreateSessionSearch failed with result [%s]"), *LexToString(ResultCode));
 
-			Op.SetError(FromEOSError(ResultCode));
+			Op.SetError(Errors::FromEOSResult(ResultCode));
 			Promise.EmplaceValue();
 			return;
 		}
@@ -867,7 +867,7 @@ TOnlineAsyncOpHandle<FFindSessions> FSessionsEOSGS::FindSessions(FFindSessions::
 	{
 		if (FindCallbackInfoResult->ResultCode != EOS_EResult::EOS_Success)
 		{
-			Op.SetError(FromEOSError(FindCallbackInfoResult->ResultCode));
+			Op.SetError(Errors::FromEOSResult(FindCallbackInfoResult->ResultCode));
 			return;
 		}
 
@@ -985,7 +985,7 @@ TOnlineAsyncOpHandle<FJoinSession> FSessionsEOSGS::JoinSession(FJoinSession::Par
 		{
 			TryRemoveSession(OpParams.SessionName);
 
-			Op.SetError(FromEOSError(Result->ResultCode));
+			Op.SetError(Errors::FromEOSResult(Result->ResultCode));
 			return;
 		}
 		
@@ -1047,7 +1047,7 @@ TResult<TSharedRef<const FSession>, FOnlineError> FSessionsEOSGS::BuildSessionFr
 	{
 		UE_LOG(LogTemp, Error, TEXT("[FSessionsEOSGS::BuildSessionFromInvite] EOS_Sessions_CopySessionHandleByInviteId failed with result [%s]"), *LexToString(CopySessionHandleByInviteIdResult));
 
-		return TResult<TSharedRef<const FSession>, FOnlineError>(FromEOSError(CopySessionHandleByInviteIdResult));
+		return TResult<TSharedRef<const FSession>, FOnlineError>(Errors::FromEOSResult(CopySessionHandleByInviteIdResult));
 	}
 }
 
@@ -1069,7 +1069,7 @@ TResult<TSharedRef<const FSession>, FOnlineError> FSessionsEOSGS::BuildSessionFr
 	{
 		UE_LOG(LogTemp, Error, TEXT("[FSessionsEOSGS::BuildSessionFromUIEvent] EOS_Sessions_CopySessionHandleByUiEventId failed with result [%s]"), *LexToString(CopySessionHandleByUiEventIdResult));
 
-		return TResult<TSharedRef<const FSession>, FOnlineError>(FromEOSError(CopySessionHandleByUiEventIdResult));
+		return TResult<TSharedRef<const FSession>, FOnlineError>(Errors::FromEOSResult(CopySessionHandleByUiEventIdResult));
 	}
 }
 
@@ -1107,7 +1107,7 @@ TOnlineAsyncOpHandle<FSendSessionInvite> FSessionsEOSGS::SendSessionInvite(FSend
 	{
 		if (Result->ResultCode != EOS_EResult::EOS_Success)
 		{
-			Op.SetError(FromEOSError(Result->ResultCode));
+			Op.SetError(Errors::FromEOSResult(Result->ResultCode));
 			return;
 		}
 
@@ -1150,7 +1150,7 @@ TOnlineAsyncOpHandle<FRejectSessionInvite> FSessionsEOSGS::RejectSessionInvite(F
 	{
 		if (Result->ResultCode != EOS_EResult::EOS_Success)
 		{
-			Op.SetError(FromEOSError(Result->ResultCode));
+			Op.SetError(Errors::FromEOSResult(Result->ResultCode));
 			return;
 		}
 
@@ -1212,7 +1212,7 @@ TOnlineAsyncOpHandle<FAddSessionMembers> FSessionsEOSGS::AddSessionMembers(FAddS
 	{
 		if (Result->ResultCode != EOS_EResult::EOS_Success)
 		{
-			Op.SetError(FromEOSError(Result->ResultCode));
+			Op.SetError(Errors::FromEOSResult(Result->ResultCode));
 			return;
 		}
 
@@ -1280,7 +1280,7 @@ TOnlineAsyncOpHandle<FRemoveSessionMembers> FSessionsEOSGS::RemoveSessionMembers
 	{
 		if (Result->ResultCode != EOS_EResult::EOS_Success)
 		{
-			Op.SetError(FromEOSError(Result->ResultCode));
+			Op.SetError(Errors::FromEOSResult(Result->ResultCode));
 			return;
 		}
 
diff --git a/Engine/Plugins/Online/Experimental/OnlineServicesEOSGS/Source/Private/Online/StatsEOSGS.cpp b/Engine/Plugins/Online/Experimental/OnlineServicesEOSGS/Source/Private/Online/StatsEOSGS.cpp
--- a/Engine/Plugins/Online/Experimental/OnlineServicesEOSGS/Source/Private/Online/StatsEOSGS.cpp
+++ b/Engine/Plugins/Online/Experimental/OnlineServicesEOSGS/Source/Private/Online/StatsEOSGS.cpp
@@ -82,7 +82,7 @@ TOnlineAsyncOpHandle<FUpdateStats> FStatsEOSGS::UpdateStats(FUpdateStats::Params
 			if (Data->ResultCode != EOS_EResult::EOS_Success)
 			{
 				UE_LOG(LogTemp, Warning, TEXT("EOS_Stats_IngestStat failed with result=[%s]"), *LexToString(Data->ResultCode));
-				InAsyncOp.SetError(FromEOSError(Data->ResultCode));
+				InAsyncOp.SetError(Errors::FromEOSResult(Data->ResultCode));
 				return;
 			}
 		});
@@ -176,7 +176,7 @@ TOnlineAsyncOpHandle<FQueryStats> FStatsEOSGS::QueryStats(FQueryStats::Params&&
 		if (Data->ResultCode != EOS_EResult::EOS_Success)
 		{
 			UE_LOG(LogTemp, Warning, TEXT("EOS_Stats_QueryStats failed with result=[%s]"), *LexToString(Data->ResultCode));
-			InAsyncOp.SetError(FromEOSError(Data->ResultCode));
+			InAsyncOp.SetError(Errors::FromEOSResult(Data->ResultCode));
 			return;
 		}
 
@@ -224,7 +224,7 @@ TOnlineAsyncOpHandle<FBatchQueryStats> FStatsEOSGS::BatchQueryStats(FBatchQueryS
 			{
 				UE_LOG(LogTemp, Warning, TEXT("EOS_Stats_QueryStats failed with result=[%s]"), *LexToString(Data->ResultCode));
 				BatchQueriedUsersStats.Reset();
-				InAsyncOp.SetError(FromEOSError(Data->ResultCode));
+				InAsyncOp.SetError(Errors::FromEOSResult(Data->ResultCode));
 				return;
 			}
 
diff --git a/Engine/Plugins/Online/Experimental/OnlineServicesEOSGS/Source/Private/Online/TitleFileEOSGS.cpp b/Engine/Plugins/Online/Experimental/OnlineServicesEOSGS/Source/Private/Online/TitleFileEOSGS.cpp
--- a/Engine/Plugins/Online/Experimental/OnlineServicesEOSGS/Source/Private/Online/TitleFileEOSGS.cpp
+++ b/Engine/Plugins/Online/Experimental/OnlineServicesEOSGS/Source/Private/Online/TitleFileEOSGS.cpp
@@ -74,7 +74,7 @@ TOnlineAsyncOpHandle<FTitleFileEnumerateFiles> FTitleFileEOSGS::EnumerateFiles(F
 		if (Data->ResultCode != EOS_EResult::EOS_Success)
 		{
 			UE_LOG(LogTemp, Warning, TEXT("EOS_TitleStorage_QueryFileList failed with result=[%s]"), *LexToString(Data->ResultCode));
-			InAsyncOp.SetError(FromEOSError(Data->ResultCode));
+			InAsyncOp.SetError(Errors::FromEOSResult(Data->ResultCode));
 			return;
 		}
 
@@ -104,7 +104,7 @@ TOnlineAsyncOpHandle<FTitleFileEnumerateFiles> FTitleFileEOSGS::EnumerateFiles(F
 			if (Result != EOS_EResult::EOS_Success)
 			{
 				UE_LOG(LogTemp, Warning, TEXT("EOS_TitleStorage_CopyFileMetadataAtIndex failed with result=[%s]"), *LexToString(Result));
-				InAsyncOp.SetError(FromEOSError(Result));
+				InAsyncOp.SetError(Errors::FromEOSResult(Result));
 				return;
 			}
 
@@ -210,7 +210,7 @@ TOnlineAsyncOpHandle<FTitleFileReadFile> FTitleFileEOSGS::ReadFile(FTitleFileRea
 		else
 		{
 			UE_LOG(LogTemp, Warning, TEXT("EOS_TitleStorage_ReadFile failed with result=[%s]"), *LexToString(Data->ResultCode));
-			InAsyncOp.SetError(FromEOSError(Data->ResultCode));
+			InAsyncOp.SetError(Errors::FromEOSResult(Data->ResultCode));
 		}
 
 		if (const EOS_HTitleStorageFileTransferRequest* RequestHandle = InAsyncOp.Data.Get<EOS_HTitleStorageFileTransferRequest>(RequestHandleKey))
diff --git a/Engine/Plugins/Online/Experimental/OnlineServicesEOSGS/Source/Private/Online/UserFileEOSGS.cpp b/Engine/Plugins/Online/Experimental/OnlineServicesEOSGS/Source/Private/Online/UserFileEOSGS.cpp
--- a/Engine/Plugins/Online/Experimental/OnlineServicesEOSGS/Source/Private/Online/UserFileEOSGS.cpp
+++ b/Engine/Plugins/Online/Experimental/OnlineServicesEOSGS/Source/Private/Online/UserFileEOSGS.cpp
@@ -65,7 +65,7 @@ TOnlineAsyncOpHandle<FUserFileEnumerateFiles> FUserFileEOSGS::EnumerateFiles(FUs
 		if (Data->ResultCode != EOS_EResult::EOS_Success)
 		{
 			UE_LOG(LogTemp, Warning, TEXT("EOS_PlayerDataStorage_QueryFileList failed with result=[%s]"), *LexToString(Data->ResultCode));
-			Op.SetError(FromEOSError(Data->ResultCode));
+			Op.SetError(Errors::FromEOSResult(Data->ResultCode));
 			return;
 		}
 
@@ -81,7 +81,7 @@ TOnlineAsyncOpHandle<FUserFileEnumerateFiles> FUserFileEOSGS::EnumerateFiles(FUs
 		if (Result != EOS_EResult::EOS_Success)
 		{
 			UE_LOG(LogTemp, Warning, TEXT("EOS_PlayerDataStorage_GetFileMetadataCount failed with result=[%s]"), *LexToString(Result));
-			Op.SetError(FromEOSError(Result));
+			Op.SetError(Errors::FromEOSResult(Result));
 			return;
 		}
 
@@ -103,7 +103,7 @@ TOnlineAsyncOpHandle<FUserFileEnumerateFiles> FUserFileEOSGS::EnumerateFiles(FUs
 			if (Result != EOS_EResult::EOS_Success)
 			{
 				UE_LOG(LogTemp, Warning, TEXT("EOS_PlayerDataStorage_CopyFileMetadataAtIndex failed with result=[%s]"), *LexToString(Result));
-				Op.SetError(FromEOSError(Result));
+				Op.SetError(Errors::FromEOSResult(Result));
 				return;
 			}
 
@@ -209,7 +209,7 @@ TOnlineAsyncOpHandle<FUserFileReadFile> FUserFileEOSGS::ReadFile(FUserFileReadFi
 		else
 		{
 			UE_LOG(LogTemp, Warning, TEXT("EOS_PlayerDataStorage_ReadFile failed with result=[%s]"), *LexToString(Data->ResultCode));
-			Op.SetError(FromEOSError(Data->ResultCode));
+			Op.SetError(Errors::FromEOSResult(Data->ResultCode));
 		}
 
 		if (const EOS_HPlayerDataStorageFileTransferRequest* RequestHandle = Op.Data.Get<EOS_HPlayerDataStorageFileTransferRequest>(RequestHandleKey))
@@ -306,7 +306,7 @@ TOnlineAsyncOpHandle<FUserFileWriteFile> FUserFileEOSGS::WriteFile(FUserFileWrit
 		else
 		{
 			UE_LOG(LogTemp, Warning, TEXT("EOS_PlayerDataStorage_WriteFile failed with result=[%s]"), *LexToString(Data->ResultCode));
-			Op.SetError(FromEOSError(Data->ResultCode));
+			Op.SetError(Errors::FromEOSResult(Data->ResultCode));
 		}
 
 		if (const EOS_HPlayerDataStorageFileTransferRequest* RequestHandle = Op.Data.Get<EOS_HPlayerDataStorageFileTransferRequest>(RequestHandleKey))
@@ -361,7 +361,7 @@ TOnlineAsyncOpHandle<FUserFileCopyFile> FUserFileEOSGS::CopyFile(FUserFileCopyFi
 		else
 		{
 			UE_LOG(LogTemp, Warning, TEXT("EOS_PlayerDataStorage_DuplicateFile failed with result=[%s]"), *LexToString(Data->ResultCode));
-			Op.SetError(FromEOSError(Data->ResultCode));
+			Op.SetError(Errors::FromEOSResult(Data->ResultCode));
 		}
 	})
 	.Enqueue(GetSerialQueue());
@@ -409,7 +409,7 @@ TOnlineAsyncOpHandle<FUserFileDeleteFile> FUserFileEOSGS::DeleteFile(FUserFileDe
 		else
 		{
 			UE_LOG(LogTemp, Warning, TEXT("EOS_PlayerDataStorage_DeleteFile failed with result=[%s]"), *LexToString(Data->ResultCode));
-			Op.SetError(FromEOSError(Data->ResultCode));
+			Op.SetError(Errors::FromEOSResult(Data->ResultCode));
 		}
 	})
 	.Enqueue(GetSerialQueue());
diff --git a/Engine/Plugins/Online/Experimental/OnlineServicesEOSGS/Source/Public/Online/OnlineErrorEOSGS.h b/Engine/Plugins/Online/Experimental/OnlineServicesEOSGS/Source/Public/Online/OnlineErrorEOSGS.h
new file mode 100644
--- /dev/null
+++ b/Engine/Plugins/Online/Experimental/OnlineServicesEOSGS/Source/Public/Online/OnlineErrorEOSGS.h
@@ -0,0 +1,47 @@
+// Copyright Epic Games, Inc. All Rights Reserved.
+
+#pragma once
+
+#include "Online/OnlineError.h"
+
+#include "eos_common.h"
+
+/*
+ *	Proper usage:
+ *		Errors::FromEOSResult(EOSResult::EOS_PlayerDataStorage_FileSizeTooLarge);
+ *	Certain EOS errors are predefined to have a common error parent type (see Internal_EOSWrapInner). i.e.
+ *		(Errors::FromEOSResult(EOSResult::EOS_NoConnection) == Errors::NoConnection()) == true
+ */
+
+namespace UE::Online::Errors {
+
+	UE_ONLINE_ERROR_CATEGORY(EOS, ThirdPartyPlugin, 0x4, "EOS")
+
+	typedef TFunction<FOnlineError(FOnlineError&& Error, EOS_EResult Result)> FErrorMapperEosFn;
+
+	ONLINESERVICESEOSGS_API FOnlineError MapCommonEOSError(FOnlineError&& Error, EOS_EResult Result);
+	ONLINESERVICESEOSGS_API ErrorCodeType ErrorCodeFromEOSResult(EOS_EResult Result);
+	ONLINESERVICESEOSGS_API FOnlineError FromEOSResult(EOS_EResult Result, FErrorMapperEosFn&& MapperFn = &MapCommonEOSError);
+
+/* UE::Online::Errors */ }
+
+// These are extern'd out of the namespace for some of the Catch API that can't see these operators if they are inside
+inline bool operator==(const UE::Online::FOnlineError& OnlineError, EOS_EResult EosResult)
+{
+	return OnlineError == UE::Online::Errors::ErrorCodeFromEOSResult(EosResult);
+}
+
+inline bool operator==(EOS_EResult EosResult, const UE::Online::FOnlineError& OnlineError)
+{
+	return OnlineError == EosResult;
+}
+
+inline bool operator!=(const UE::Online::FOnlineError& OnlineError, EOS_EResult EosResult)
+{
+	return !(OnlineError == EosResult);
+}
+
+inline bool operator!=(EOS_EResult EosResult, const UE::Online::FOnlineError& OnlineError)
+{
+	return OnlineError != EosResult;
+}
diff --git a/Engine/Plugins/Online/Experimental/OnlineServicesEOSGS/Source/Public/Online/OnlineServicesEOSGSTypes.h b/Engine/Plugins/Online/Experimental/OnlineServicesEOSGS/Source/Public/Online/OnlineServicesEOSGSTypes.h
--- a/Engine/Plugins/Online/Experimental/OnlineServicesEOSGS/Source/Public/Online/OnlineServicesEOSGSTypes.h
+++ b/Engine/Plugins/Online/Experimental/OnlineServicesEOSGS/Source/Public/Online/OnlineServicesEOSGSTypes.h
@@ -3,19 +3,13 @@
 #pragma once
 
 #include "Online/OnlineAsyncOp.h"
-#include "Online/OnlineError.h"
+#include "Online/OnlineErrorEOSGS.h"
 #include "Async/Future.h"
 
 #include "EOSShared.h"
 
 namespace UE::Online {
 
-inline FOnlineError FromEOSError(EOS_EResult ResultCode)
-{
-	// Todo: make this for real.
-	return FOnlineError(Errors::Unknown());
-}
-
 namespace Private {
 
 /** Class to handle all callbacks generically using a future to forward callback results */
diff --git a/Engine/Plugins/Online/Experimental/OnlineServicesOSSAdapter/Source/Public/Online/ErrorsOSSAdapter.h b/Engine/Plugins/Online/Experimental/OnlineServicesOSSAdapter/Source/Public/Online/ErrorsOSSAdapter.h
--- a/Engine/Plugins/Online/Experimental/OnlineServicesOSSAdapter/Source/Public/Online/ErrorsOSSAdapter.h
+++ b/Engine/Plugins/Online/Experimental/OnlineServicesOSSAdapter/Source/Public/Online/ErrorsOSSAdapter.h
@@ -7,7 +7,7 @@ using FOnlineErrorOss = ::FOnlineError;
 
 namespace UE::Online::Errors {
 
-	UE_ONLINE_ERROR_CATEGORY(Oss, Engine, 0x5, "Oss")
+	UE_ONLINE_ERROR_CATEGORY(Oss, Engine, 0x7, "Oss")
 	FOnlineError FromOssError(const FOnlineErrorOss& Result);
 	
 	// if you would like to add a parser that transforms custom platform errors into readable online service types, you can call this in your platform's error handler
