Fix build issues caused by error consolidation.

diff --git a/Engine/Plugins/Online/Experimental/OnlineServicesEOSGS/Source/Private/Online/LeaderboardsEOSGS.cpp b/Engine/Plugins/Online/Experimental/OnlineServicesEOSGS/Source/Private/Online/LeaderboardsEOSGS.cpp
--- a/Engine/Plugins/Online/Experimental/OnlineServicesEOSGS/Source/Private/Online/LeaderboardsEOSGS.cpp
+++ b/Engine/Plugins/Online/Experimental/OnlineServicesEOSGS/Source/Private/Online/LeaderboardsEOSGS.cpp
@@ -2,6 +2,7 @@
 
 #include "Online/LeaderboardsEOSGS.h"
 #include "Online/AuthEOSGS.h"
+#include "Online/OnlineErrorEOSGS.h"
 #include "Online/OnlineIdEOSGS.h"
 #include "Online/OnlineServicesEOSGS.h"
 #include "Online/OnlineServicesEOSGSTypes.h"
@@ -74,7 +75,7 @@ TOnlineAsyncOpHandle<FReadEntriesForUsers> FLeaderboardsEOSGS::ReadEntriesForUse
 		if (Data->ResultCode != EOS_EResult::EOS_Success)
 		{
 			UE_LOG(LogTemp, Warning, TEXT("EOS_Leaderboards_IngestStat failed with result=[%s]"), *LexToString(Data->ResultCode));
-			InAsyncOp.SetError(FromEOSError(Data->ResultCode));
+			InAsyncOp.SetError(Errors::FromEOSResult(Data->ResultCode));
 			return;
 		}
 
@@ -183,7 +184,7 @@ TOnlineAsyncOpHandle<FReadEntriesAroundRank> FLeaderboardsEOSGS::ReadEntriesArou
 		if (Data->ResultCode != EOS_EResult::EOS_Success)
 		{
 			UE_LOG(LogTemp, Warning, TEXT("EOS_Leaderboards_QueryLeaderboardRanks failed with result=[%s]"), *LexToString(Data->ResultCode));
-			InAsyncOp.SetError(FromEOSError(Data->ResultCode));
+			InAsyncOp.SetError(Errors::FromEOSResult(Data->ResultCode));
 			return;
 		}
 
@@ -224,7 +225,7 @@ TOnlineAsyncOpHandle<FReadEntriesAroundUser> FLeaderboardsEOSGS::ReadEntriesArou
 		if (Data->ResultCode != EOS_EResult::EOS_Success)
 		{
 			UE_LOG(LogTemp, Warning, TEXT("EOS_Leaderboards_QueryLeaderboardRanks failed with result=[%s]"), *LexToString(Data->ResultCode));
-			InAsyncOp.SetError(FromEOSError(Data->ResultCode));
+			InAsyncOp.SetError(Errors::FromEOSResult(Data->ResultCode));
 			return;
 		}
 
diff --git a/Engine/Plugins/Online/Experimental/OnlineServicesEOSGS/Source/Private/Online/OnlineErrorEOSGS.cpp b/Engine/Plugins/Online/Experimental/OnlineServicesEOSGS/Source/Private/Online/OnlineErrorEOSGS.cpp
--- a/Engine/Plugins/Online/Experimental/OnlineServicesEOSGS/Source/Private/Online/OnlineErrorEOSGS.cpp
+++ b/Engine/Plugins/Online/Experimental/OnlineServicesEOSGS/Source/Private/Online/OnlineErrorEOSGS.cpp
@@ -3,6 +3,8 @@
 #include "Online/OnlineErrorEOSGS.h"
 #include "Online/OnlineErrorDefinitions.h"
 
+#include "EOSShared.h"
+
 namespace UE::Online::Errors {
 
 class FOnlineErrorDetailsEOS : public IOnlineErrorDetails
@@ -22,12 +24,14 @@ const TSharedRef<const IOnlineErrorDetails, ESPMode::ThreadSafe>& FOnlineErrorDe
 
 FText FOnlineErrorDetailsEOS::GetText(const FOnlineError& OnlineError) const
 {
-	return FText::FromString(LexToString(OnlineError.GetValue()));
+	const EOS_EResult EosResult = EOS_EResult(OnlineError.GetValue());
+	return FText::FromString(LexToString(EosResult));
 }
 
 FString FOnlineErrorDetailsEOS::GetLogString(const FOnlineError& OnlineError) const
 {
-	return LexToString(OnlineError.GetValue());
+	const EOS_EResult EosResult = EOS_EResult(OnlineError.GetValue());
+	return LexToString(EosResult);
 }
 
 /** This doesn't wrap every single error in EOS, only the common ones that are strongly related to common errors */
