Add static asserts to Auth and Lobbies for detecting EOS structure changes.

diff --git a/Engine/Plugins/Online/Experimental/OnlineServicesEOS/Source/Private/Online/AuthEOS.cpp b/Engine/Plugins/Online/Experimental/OnlineServicesEOS/Source/Private/Online/AuthEOS.cpp
--- a/Engine/Plugins/Online/Experimental/OnlineServicesEOS/Source/Private/Online/AuthEOS.cpp
+++ b/Engine/Plugins/Online/Experimental/OnlineServicesEOS/Source/Private/Online/AuthEOS.cpp
@@ -205,6 +205,7 @@ TOnlineAsyncOpHandle<FAuthLogin> FAuthEOS::Login(FAuthLogin::Params&& Params)
 		Options.ApiVersion = EOS_USERINFO_COPYUSERINFO_API_LATEST;
 		Options.LocalUserId = AccountInfoEOS->EpicAccountId;
 		Options.TargetUserId = AccountInfoEOS->EpicAccountId;
+		static_assert(EOS_USERINFO_COPYUSERINFO_API_LATEST == 3, "EOS_UserInfo_CopyUserInfoOptions updated, check new fields");
 
 		EOS_UserInfo* UserInfo = nullptr;
 
@@ -290,6 +291,7 @@ TOnlineAsyncOpHandle<FAuthQueryExternalServerAuthTicket> FAuthEOS::QueryExternal
 
 		EOS_Auth_CopyUserAuthTokenOptions CopyUserAuthTokenOptions = {};
 		CopyUserAuthTokenOptions.ApiVersion = EOS_AUTH_COPYUSERAUTHTOKEN_API_LATEST;
+		static_assert(EOS_AUTH_COPYUSERAUTHTOKEN_API_LATEST == 1, "EOS_Auth_CopyUserAuthTokenOptions updated, check new fields");
 
 		EOS_Auth_Token* AuthToken = nullptr;
 
@@ -417,6 +419,7 @@ TFuture<TArray<FOnlineAccountIdHandle>> FAuthEOS::ResolveAccountIds(const FOnlin
 	Options.AccountIdType = EOS_EExternalAccountType::EOS_EAT_EPIC;
 	Options.ExternalAccountIds = (const char**)EpicAccountIdStrPtrs.GetData();
 	Options.ExternalAccountIdCount = 1;
+	static_assert(EOS_CONNECT_QUERYEXTERNALACCOUNTMAPPINGS_API_LATEST == 1, "EOS_Connect_QueryExternalAccountMappingsOptions updated, check new fields");
 
 	EOS_Async(EOS_Connect_QueryExternalAccountMappings, ConnectHandle, Options,
 	[this, WeakThis = AsWeak(), InEpicAccountIds, Promise = MoveTemp(Promise)](const EOS_Connect_QueryExternalAccountMappingsCallbackInfo* Data) mutable -> void
@@ -431,6 +434,7 @@ TFuture<TArray<FOnlineAccountIdHandle>> FAuthEOS::ResolveAccountIds(const FOnlin
 				Options.ApiVersion = EOS_CONNECT_GETEXTERNALACCOUNTMAPPING_API_LATEST;
 				Options.LocalUserId = Data->LocalUserId;
 				Options.AccountIdType = EOS_EExternalAccountType::EOS_EAT_EPIC;
+				static_assert(EOS_CONNECT_GETEXTERNALACCOUNTMAPPING_API_LATEST == 1, "EOS_Connect_GetExternalAccountMappingsOptions updated, check new fields");
 
 				for (const EOS_EpicAccountId EpicAccountId : InEpicAccountIds)
 				{
@@ -500,6 +504,8 @@ TFuture<TArray<FOnlineAccountIdHandle>> FAuthEOS::ResolveAccountIds(const FOnlin
 	Options.LocalUserId = GetProductUserIdChecked(LocalUserId);
 	Options.ProductUserIds = MissingProductUserIds.GetData();
 	Options.ProductUserIdCount = MissingProductUserIds.Num();
+	static_assert(EOS_CONNECT_QUERYPRODUCTUSERIDMAPPINGS_API_LATEST == 2, "EOS_Connect_QueryProductUserIdMappingsOptions updated, check new fields");
+	
 	EOS_Async(EOS_Connect_QueryProductUserIdMappings, ConnectHandle, Options,
 	[this, WeakThis = AsWeak(), InProductUserIds, Promise = MoveTemp(Promise)](const EOS_Connect_QueryProductUserIdMappingsCallbackInfo* Data) mutable -> void
 	{
@@ -512,6 +518,7 @@ TFuture<TArray<FOnlineAccountIdHandle>> FAuthEOS::ResolveAccountIds(const FOnlin
 				Options.ApiVersion = EOS_CONNECT_GETPRODUCTUSERIDMAPPING_API_LATEST;
 				Options.LocalUserId = Data->LocalUserId;
 				Options.AccountIdType = EOS_EExternalAccountType::EOS_EAT_EPIC;
+				static_assert(EOS_CONNECT_GETPRODUCTUSERIDMAPPING_API_LATEST == 1, "EOS_Connect_GetProductUserIdMappingOptions updated, check new fields");
 
 				for (const EOS_ProductUserId ProductUserId : InProductUserIds)
 				{
diff --git a/Engine/Plugins/Online/Experimental/OnlineServicesEOSGS/Source/Private/Online/AuthEOSGS.cpp b/Engine/Plugins/Online/Experimental/OnlineServicesEOSGS/Source/Private/Online/AuthEOSGS.cpp
--- a/Engine/Plugins/Online/Experimental/OnlineServicesEOSGS/Source/Private/Online/AuthEOSGS.cpp
+++ b/Engine/Plugins/Online/Experimental/OnlineServicesEOSGS/Source/Private/Online/AuthEOSGS.cpp
@@ -151,15 +151,18 @@ FEOSConnectLoginOptions& FEOSConnectLoginOptions::operator=(FEOSConnectLoginOpti
 FEOSConnectLoginOptions::FEOSConnectLoginOptions()
 {
 	// EOS_Connect_LoginOptions init
+	static_assert(EOS_CONNECT_LOGIN_API_LATEST == 2, "EOS_Connect_LoginOptions updated, check new fields");
 	ApiVersion = EOS_CONNECT_LOGIN_API_LATEST;
 	Credentials = &CredentialsData;
 	UserLoginInfo = nullptr;
 
 	// EOS_Connect_Credentials init
+	static_assert(EOS_CONNECT_CREDENTIALS_API_LATEST == 1, "EOS_Connect_Credentials updated, check new fields");
 	CredentialsData.ApiVersion = EOS_CONNECT_CREDENTIALS_API_LATEST;
 	CredentialsData.Token = nullptr;
 
 	// EOS_Connect_UserLoginInfo init
+	static_assert(EOS_CONNECT_USERLOGININFO_API_LATEST == 1, "EOS_Connect_UserLoginInfo updated, check new fields");
 	UserLoginInfoData.ApiVersion = EOS_CONNECT_USERLOGININFO_API_LATEST;
 	UserLoginInfoData.DisplayName = DisplayNameUtf8;
 }
@@ -332,11 +335,13 @@ FEOSAuthLoginOptions& FEOSAuthLoginOptions::operator=(FEOSAuthLoginOptions&& Oth
 FEOSAuthLoginOptions::FEOSAuthLoginOptions()
 {
 	// EOS_Auth_LoginOptions init
+	static_assert(EOS_AUTH_LOGIN_API_LATEST == 2, "EOS_Auth_LoginOptions updated, check new fields");
 	ApiVersion = EOS_AUTH_LOGIN_API_LATEST;
 	Credentials = &CredentialsData;
 	ScopeFlags = EOS_EAuthScopeFlags::EOS_AS_NoFlags;
 
 	// EOS_Auth_Credentials init
+	static_assert(EOS_AUTH_CREDENTIALS_API_LATEST == 3, "EOS_Auth_Credentials updated, check new fields");
 	CredentialsData.ApiVersion = EOS_AUTH_CREDENTIALS_API_LATEST;
 	CredentialsData.Id = nullptr;
 	CredentialsData.Token = nullptr;
@@ -847,6 +852,7 @@ TOnlineAsyncOpHandle<FAuthLogout> FAuthEOSGS::Logout(FAuthLogout::Params&& Param
 			EOS_Auth_DeletePersistentAuthOptions DeletePersistentAuthOptions = { 0 };
 			DeletePersistentAuthOptions.ApiVersion = EOS_AUTH_DELETEPERSISTENTAUTH_API_LATEST;
 			DeletePersistentAuthOptions.RefreshToken = nullptr; // TODO: Is this needed?  Docs say it's needed for consoles
+			static_assert(EOS_AUTH_DELETEPERSISTENTAUTH_API_LATEST == 2, "EOS_Auth_DeletePersistentAuthOptions updated, check new fields");
 
 			EOS_Async(EOS_Auth_DeletePersistentAuth, AuthHandle, DeletePersistentAuthOptions,
 			[Promise = MoveTemp(Promise)](const EOS_Auth_DeletePersistentAuthCallbackInfo* Data) mutable -> void
@@ -1009,6 +1015,7 @@ TFuture<TDefaultErrorResult<FAuthLoginEASImpl>> FAuthEOSGS::LoginEASImpl(const F
 			EOS_Auth_LinkAccountOptions LinkAccountOptions = {};
 			LinkAccountOptions.ApiVersion = EOS_AUTH_LINKACCOUNT_API_LATEST;
 			LinkAccountOptions.ContinuanceToken = Data->ContinuanceToken;
+			static_assert(EOS_AUTH_LINKACCOUNT_API_LATEST == 1, "EOS_Auth_LinkAccountOptions updated, check new fields");
 
 			EOS_Async(EOS_Auth_LinkAccount, AuthHandle, LinkAccountOptions,
 			[Promise = MoveTemp(Promise)](const EOS_Auth_LinkAccountCallbackInfo* Data) mutable -> void
@@ -1039,6 +1046,7 @@ TFuture<TDefaultErrorResult<FAuthLoginEASImpl>> FAuthEOSGS::LoginEASImpl(const F
 			{
 				EOS_Auth_DeletePersistentAuthOptions DeletePersistentAuthOptions = {};
 				DeletePersistentAuthOptions.ApiVersion = EOS_AUTH_DELETEPERSISTENTAUTH_API_LATEST;
+				static_assert(EOS_AUTH_DELETEPERSISTENTAUTH_API_LATEST == 2, "EOS_Auth_DeletePersistentAuthOptions updated, check new fields");
 
 				EOS_Async(EOS_Auth_DeletePersistentAuth, AuthHandle, DeletePersistentAuthOptions,
 				[ResolvedError = MoveTemp(ResolvedError), Promise = MoveTemp(Promise)](const EOS_Auth_DeletePersistentAuthCallbackInfo* Data) mutable -> void
@@ -1070,6 +1078,7 @@ TFuture<TDefaultErrorResult<FAuthLogoutEASImpl>> FAuthEOSGS::LogoutEASImpl(const
 	EOS_Auth_LogoutOptions LogoutOptions = {};
 	LogoutOptions.ApiVersion = EOS_AUTH_LOGOUT_API_LATEST;
 	LogoutOptions.LocalUserId = LogoutParams.EpicAccountId;
+	static_assert(EOS_AUTH_LOGOUT_API_LATEST == 1, "EOS_Auth_LogoutOptions updated, check new fields");
 
 	EOS_Async(EOS_Auth_Logout, AuthHandle, LogoutOptions,
 	[Promise = MoveTemp(Promise)](const EOS_Auth_LogoutCallbackInfo* Data) mutable -> void
@@ -1094,6 +1103,7 @@ TDefaultErrorResult<FAuthGetExternalAuthTokenImpl> FAuthEOSGS::GetExternalAuthTo
 	EOS_Auth_CopyIdTokenOptions CopyIdTokenOptions = {};
 	CopyIdTokenOptions.ApiVersion = EOS_AUTH_COPYIDTOKEN_API_LATEST;
 	CopyIdTokenOptions.AccountId = Params.EpicAccountId;
+	static_assert(EOS_AUTH_COPYIDTOKEN_API_LATEST == 1, "EOS_Auth_CopyIdTokenOptions updated, check new fields");
 
 	EOS_Auth_IdToken* IdToken = nullptr;
 	EOS_EResult Result = EOS_Auth_CopyIdToken(AuthHandle, &CopyIdTokenOptions, &IdToken);
@@ -1140,6 +1150,7 @@ TFuture<TDefaultErrorResult<FAuthLoginConnectImpl>> FAuthEOSGS::LoginConnectImpl
 			EOS_Connect_CreateUserOptions ConnectCreateUserOptions = { };
 			ConnectCreateUserOptions.ApiVersion = EOS_CONNECT_CREATEUSER_API_LATEST;
 			ConnectCreateUserOptions.ContinuanceToken = Data->ContinuanceToken;
+			static_assert(EOS_CONNECT_CREATEUSER_API_LATEST == 1, "EOS_Connect_CreateUserOptions updated, check new fields");
 
 			EOS_Async(EOS_Connect_CreateUser, ConnectHandle, ConnectCreateUserOptions,
 			[Promise = MoveTemp(Promise)](const EOS_Connect_CreateUserCallbackInfo* Data) mutable -> void
@@ -1417,6 +1428,7 @@ void FAuthEOSGS::RegisterHandlers()
 		&EOS_Connect_AddNotifyLoginStatusChanged,
 		&EOS_Connect_RemoveNotifyLoginStatusChanged,
 		&FAuthEOSGS::OnConnectLoginStatusChanged);
+	static_assert(EOS_CONNECT_ADDNOTIFYLOGINSTATUSCHANGED_API_LATEST == 1, "EOS_Connect_AddNotifyLoginStatusChanged updated, check new fields");
 
 	// Notification of a pending external token expiration ~10 minutes.
 	OnConnectAuthNotifyExpirationEOSEventRegistration = EOS_RegisterComponentEventHandler(
@@ -1426,6 +1438,7 @@ void FAuthEOSGS::RegisterHandlers()
 		&EOS_Connect_AddNotifyAuthExpiration,
 		&EOS_Connect_RemoveNotifyAuthExpiration,
 		&FAuthEOSGS::OnConnectAuthNotifyExpiration);
+	static_assert(EOS_CONNECT_ONAUTHEXPIRATIONCALLBACK_API_LATEST == 1, "EOS_Connect_AddNotifyAuthExpiration updated, check new fields");
 
 	// Register for EAS connection status updates.
 	OnConnectAuthNotifyExpirationEOSEventRegistration = EOS_RegisterComponentEventHandler(
@@ -1435,6 +1448,7 @@ void FAuthEOSGS::RegisterHandlers()
 		&EOS_Auth_AddNotifyLoginStatusChanged,
 		&EOS_Auth_RemoveNotifyLoginStatusChanged,
 		&FAuthEOSGS::OnEASLoginStatusChanged);
+	static_assert(EOS_AUTH_ADDNOTIFYLOGINSTATUSCHANGED_API_LATEST == 1, "EOS_Auth_AddNotifyLoginStatusChanged updated, check new fields");
 }
 
 void FAuthEOSGS::UnregisterHandlers()
diff --git a/Engine/Plugins/Online/Experimental/OnlineServicesEOSGS/Source/Private/Online/LobbiesEOSGS.cpp b/Engine/Plugins/Online/Experimental/OnlineServicesEOSGS/Source/Private/Online/LobbiesEOSGS.cpp
--- a/Engine/Plugins/Online/Experimental/OnlineServicesEOSGS/Source/Private/Online/LobbiesEOSGS.cpp
+++ b/Engine/Plugins/Online/Experimental/OnlineServicesEOSGS/Source/Private/Online/LobbiesEOSGS.cpp
@@ -117,6 +117,7 @@ TOnlineAsyncOpHandle<FCreateLobby> FLobbiesEOSGS::CreateLobby(FCreateLobby::Para
 		CreateLobbyOptions.BucketId = BucketTanslator.GetBucketIdEOS();
 		CreateLobbyOptions.bDisableHostMigration = false; // todo: handle
 		CreateLobbyOptions.bEnableRTCRoom = false; // todo: handle
+		static_assert(EOS_LOBBY_CREATELOBBY_API_LATEST == 8, "EOS_Lobby_CreateLobbyOptions updated, check new fields");
 
 		EOS_Async(EOS_Lobby_CreateLobby, LobbyPrerequisites->LobbyInterfaceHandle, CreateLobbyOptions, MoveTemp(Promise)); 
 	})
@@ -1090,6 +1091,7 @@ void FLobbiesEOSGS::RegisterHandlers()
 		&EOS_Lobby_AddNotifyLobbyUpdateReceived,
 		&EOS_Lobby_RemoveNotifyLobbyUpdateReceived,
 		&FLobbiesEOSGS::HandleLobbyUpdated);
+	static_assert(EOS_LOBBY_ADDNOTIFYLOBBYUPDATERECEIVED_API_LATEST == 1, "EOS_Lobby_AddNotifyLobbyUpdateReceived updated, check new fields");
 
 	// Register for lobby member updates.
 	OnLobbyMemberUpdatedEOSEventRegistration = EOS_RegisterComponentEventHandler(
@@ -1099,6 +1101,7 @@ void FLobbiesEOSGS::RegisterHandlers()
 		&EOS_Lobby_AddNotifyLobbyMemberUpdateReceived,
 		&EOS_Lobby_RemoveNotifyLobbyMemberUpdateReceived,
 		&FLobbiesEOSGS::HandleLobbyMemberUpdated);
+	static_assert(EOS_LOBBY_ADDNOTIFYLOBBYMEMBERUPDATERECEIVED_API_LATEST == 1, "EOS_Lobby_AddNotifyLobbyMemberUpdateReceived updated, check new fields");
 
 	// Register for lobby member status changed.
 	OnLobbyMemberStatusReceivedEOSEventRegistration = EOS_RegisterComponentEventHandler(
@@ -1108,6 +1111,7 @@ void FLobbiesEOSGS::RegisterHandlers()
 		&EOS_Lobby_AddNotifyLobbyMemberStatusReceived,
 		&EOS_Lobby_RemoveNotifyLobbyMemberStatusReceived,
 		&FLobbiesEOSGS::HandleLobbyMemberStatusReceived);
+	static_assert(EOS_LOBBY_ADDNOTIFYLOBBYMEMBERSTATUSRECEIVED_API_LATEST == 1, "EOS_Lobby_AddNotifyLobbyMemberStatusReceived updated, check new fields");
 
 	// Register for lobby invite received.
 	OnLobbyInviteReceivedEOSEventRegistration = EOS_RegisterComponentEventHandler(
@@ -1117,6 +1121,7 @@ void FLobbiesEOSGS::RegisterHandlers()
 		&EOS_Lobby_AddNotifyLobbyInviteReceived,
 		&EOS_Lobby_RemoveNotifyLobbyInviteReceived,
 		&FLobbiesEOSGS::HandleLobbyInviteReceived);
+	static_assert(EOS_LOBBY_ADDNOTIFYLOBBYINVITERECEIVED_API_LATEST == 1, "EOS_Lobby_AddNotifyLobbyInviteReceived updated, check new fields");
 
 	// Register for lobby invite accepted.
 	OnLobbyInviteAcceptedEOSEventRegistration = EOS_RegisterComponentEventHandler(
@@ -1126,6 +1131,7 @@ void FLobbiesEOSGS::RegisterHandlers()
 		&EOS_Lobby_AddNotifyLobbyInviteAccepted,
 		&EOS_Lobby_RemoveNotifyLobbyInviteAccepted,
 		&FLobbiesEOSGS::HandleLobbyInviteAccepted);
+	static_assert(EOS_LOBBY_ADDNOTIFYLOBBYINVITEACCEPTED_API_LATEST == 1, "EOS_Lobby_AddNotifyLobbyInviteAccepted updated, check new fields");
 
 	// Register for join lobby accepted via overlay.
 	OnJoinLobbyAcceptedEOSEventRegistration = EOS_RegisterComponentEventHandler(
@@ -1135,6 +1141,7 @@ void FLobbiesEOSGS::RegisterHandlers()
 		&EOS_Lobby_AddNotifyJoinLobbyAccepted,
 		&EOS_Lobby_RemoveNotifyJoinLobbyAccepted,
 		&FLobbiesEOSGS::HandleJoinLobbyAccepted);
+	static_assert(EOS_LOBBY_ADDNOTIFYJOINLOBBYACCEPTED_API_LATEST == 1, "EOS_Lobby_AddNotifyJoinLobbyAccepted updated, check new fields");
 }
 
 void FLobbiesEOSGS::UnregisterHandlers()
@@ -1338,6 +1345,8 @@ TOnlineAsyncOpHandle<FLobbiesJoinLobbyMemberImpl> FLobbiesEOSGS::JoinLobbyMember
 		JoinLobbyOptions.LocalUserId = GetProductUserIdChecked(Params.LocalUserId);
 		JoinLobbyOptions.bPresenceEnabled = false; // todo
 		JoinLobbyOptions.LocalRTCOptions = nullptr;
+		static_assert(EOS_LOBBY_JOINLOBBY_API_LATEST == 3, "EOS_Lobby_JoinLobbyOptions updated, check new fields");
+		
 		EOS_Async(EOS_Lobby_JoinLobby, LobbyPrerequisites->LobbyInterfaceHandle, JoinLobbyOptions, MoveTemp(Promise));
 	})
 	.Then([this](TOnlineAsyncOp<FLobbiesJoinLobbyMemberImpl>& InAsyncOp, const EOS_Lobby_JoinLobbyCallbackInfo* Data)
@@ -1418,6 +1427,7 @@ TFuture<TDefaultErrorResult<FLobbiesLeaveLobbyImpl>> FLobbiesEOSGS::LeaveLobbyIm
 	LeaveLobbyOptions.ApiVersion = EOS_LOBBY_LEAVELOBBY_API_LATEST;
 	LeaveLobbyOptions.LocalUserId = GetProductUserIdChecked(Params.LocalUserId);
 	LeaveLobbyOptions.LobbyId = Params.LobbyData->GetLobbyIdEOS();
+	static_assert(EOS_LOBBY_LEAVELOBBY_API_LATEST == 1, "EOS_Lobby_LeaveLobbyOptions updated, check new fields");
 
 	TPromise<TDefaultErrorResult<FLobbiesLeaveLobbyImpl>> Promise;
 	TFuture<TDefaultErrorResult<FLobbiesLeaveLobbyImpl>> Future = Promise.GetFuture();
@@ -1455,6 +1465,7 @@ TFuture<TDefaultErrorResult<FLobbiesDestroyLobbyImpl>> FLobbiesEOSGS::DestroyLob
 	DestroyLobbyOptions.ApiVersion = EOS_LOBBY_DESTROYLOBBY_API_LATEST;
 	DestroyLobbyOptions.LocalUserId = GetProductUserIdChecked(Params.LocalUserId);
 	DestroyLobbyOptions.LobbyId = Params.LobbyData->GetLobbyIdEOS();
+	static_assert(EOS_LOBBY_DESTROYLOBBY_API_LATEST == 1, "EOS_Lobby_DestroyLobbyOptions updated, check new fields");
 
 	TPromise<TDefaultErrorResult<FLobbiesDestroyLobbyImpl>> Promise;
 	TFuture<TDefaultErrorResult<FLobbiesDestroyLobbyImpl>> Future = Promise.GetFuture();
@@ -1493,6 +1504,7 @@ TFuture<TDefaultErrorResult<FLobbiesInviteLobbyMemberImpl>> FLobbiesEOSGS::Invit
 	SendInviteOptions.LocalUserId = GetProductUserIdChecked(Params.LocalUserId);
 	SendInviteOptions.TargetUserId = GetProductUserIdChecked(Params.TargetUserId);
 	SendInviteOptions.LobbyId = Params.LobbyData->GetLobbyIdEOS();
+	static_assert(EOS_LOBBY_SENDINVITE_API_LATEST == 1, "EOS_Lobby_SendInviteOptions updated, check new fields");
 
 	TPromise<TDefaultErrorResult<FLobbiesInviteLobbyMemberImpl>> Promise;
 	TFuture<TDefaultErrorResult<FLobbiesInviteLobbyMemberImpl>> Future = Promise.GetFuture();
@@ -1533,6 +1545,7 @@ TFuture<TDefaultErrorResult<FLobbiesDeclineLobbyInvitationImpl>> FLobbiesEOSGS::
 	RejectInviteOptions.ApiVersion = EOS_LOBBY_REJECTINVITE_API_LATEST;
 	RejectInviteOptions.LocalUserId = GetProductUserIdChecked(Params.LocalUserId);
 	RejectInviteOptions.InviteId = InviteData->GetInviteIdEOS();
+	static_assert(EOS_LOBBY_REJECTINVITE_API_LATEST == 1, "EOS_Lobby_RejectInviteOptions updated, check new fields");
 
 	TPromise<TDefaultErrorResult<FLobbiesDeclineLobbyInvitationImpl>> Promise;
 	TFuture<TDefaultErrorResult<FLobbiesDeclineLobbyInvitationImpl>> Future = Promise.GetFuture();
@@ -1576,6 +1589,7 @@ TFuture<TDefaultErrorResult<FLobbiesKickLobbyMemberImpl>> FLobbiesEOSGS::KickLob
 	KickMemberOptions.LocalUserId = GetProductUserIdChecked(Params.LocalUserId);
 	KickMemberOptions.TargetUserId = GetProductUserIdChecked(Params.TargetUserId);
 	KickMemberOptions.LobbyId = Params.LobbyData->GetLobbyIdEOS();
+	static_assert(EOS_LOBBY_KICKMEMBER_API_LATEST == 1, "EOS_Lobby_KickMemberOptions updated, check new fields");
 
 	TPromise<TDefaultErrorResult<FLobbiesKickLobbyMemberImpl>> Promise;
 	TFuture<TDefaultErrorResult<FLobbiesKickLobbyMemberImpl>> Future = Promise.GetFuture();
@@ -1616,6 +1630,7 @@ TFuture<TDefaultErrorResult<FLobbiesPromoteLobbyMemberImpl>> FLobbiesEOSGS::Prom
 	PromoteMemberOptions.LocalUserId = GetProductUserIdChecked(Params.LocalUserId);
 	PromoteMemberOptions.TargetUserId = GetProductUserIdChecked(Params.TargetUserId);
 	PromoteMemberOptions.LobbyId = Params.LobbyData->GetLobbyIdEOS();
+	static_assert(EOS_LOBBY_PROMOTEMEMBER_API_LATEST == 1, "EOS_Lobby_PromoteMemberOptions updated, check new fields");
 
 	TPromise<TDefaultErrorResult<FLobbiesPromoteLobbyMemberImpl>> Promise;
 	TFuture<TDefaultErrorResult<FLobbiesPromoteLobbyMemberImpl>> Future = Promise.GetFuture();
diff --git a/Engine/Plugins/Online/Experimental/OnlineServicesEOSGS/Source/Private/Online/LobbiesEOSGSTypes.cpp b/Engine/Plugins/Online/Experimental/OnlineServicesEOSGS/Source/Private/Online/LobbiesEOSGSTypes.cpp
--- a/Engine/Plugins/Online/Experimental/OnlineServicesEOSGS/Source/Private/Online/LobbiesEOSGSTypes.cpp
+++ b/Engine/Plugins/Online/Experimental/OnlineServicesEOSGS/Source/Private/Online/LobbiesEOSGSTypes.cpp
@@ -67,6 +67,7 @@ FLobbyAttributeTranslator<ELobbyTranslationType::ToService>::FLobbyAttributeTran
 {
 	AttributeData.ApiVersion = EOS_LOBBY_ATTRIBUTEDATA_API_LATEST;
 	AttributeData.Key = KeyConverterStorage.Get();
+	static_assert(EOS_LOBBY_ATTRIBUTEDATA_API_LATEST == 1, "EOS_Lobby_AttributeData updated, check new fields");
 
 	if (FromAttributeData.VariantData.IsType<FString>())
 	{
@@ -155,6 +156,7 @@ TDefaultErrorResultInternal<TSharedRef<FLobbyDetailsEOS>> FLobbyDetailsEOS::Crea
 	Options.ApiVersion = EOS_LOBBY_COPYLOBBYDETAILSHANDLE_API_LATEST;
 	Options.LobbyId = LobbyId;
 	Options.LocalUserId = GetProductUserIdChecked(LocalUserId);
+	static_assert(EOS_LOBBY_COPYLOBBYDETAILSHANDLE_API_LATEST == 1, "EOS_Lobby_CopyLobbyDetailsHandleOptions updated, check new fields");
 
 	EOS_EResult EOSResult = EOS_Lobby_CopyLobbyDetailsHandle(Prerequisites->LobbyInterfaceHandle, &Options, &LobbyDetailsHandle);
 	if (EOSResult != EOS_EResult::EOS_Success)
@@ -182,6 +184,7 @@ TDefaultErrorResultInternal<TSharedRef<FLobbyDetailsEOS>> FLobbyDetailsEOS::Crea
 	EOS_Lobby_CopyLobbyDetailsHandleByInviteIdOptions Options;
 	Options.ApiVersion = EOS_LOBBY_COPYLOBBYDETAILSHANDLEBYINVITEID_API_LATEST;
 	Options.InviteId = InviteId;
+	static_assert(EOS_LOBBY_COPYLOBBYDETAILSHANDLEBYINVITEID_API_LATEST == 1, "EOS_Lobby_CopyLobbyDetailsHandleByInviteIdOptions updated, check new fields");
 
 	EOS_EResult EOSResult = EOS_Lobby_CopyLobbyDetailsHandleByInviteId(Prerequisites->LobbyInterfaceHandle, &Options, &LobbyDetailsHandle);
 	if (EOSResult != EOS_EResult::EOS_Success)
@@ -209,6 +212,7 @@ TDefaultErrorResultInternal<TSharedRef<FLobbyDetailsEOS>> FLobbyDetailsEOS::Crea
 	EOS_Lobby_CopyLobbyDetailsHandleByUiEventIdOptions Options;
 	Options.ApiVersion = EOS_LOBBY_COPYLOBBYDETAILSHANDLEBYUIEVENTID_API_LATEST;
 	Options.UiEventId = UiEventId;
+	static_assert(EOS_LOBBY_COPYLOBBYDETAILSHANDLEBYUIEVENTID_API_LATEST == 1, "EOS_Lobby_CopyLobbyDetailsHandleByUiEventIdOptions updated, check new fields");
 
 	EOS_EResult EOSResult = EOS_Lobby_CopyLobbyDetailsHandleByUiEventId(Prerequisites->LobbyInterfaceHandle, &Options, &LobbyDetailsHandle);
 	if (EOSResult != EOS_EResult::EOS_Success)
@@ -237,6 +241,7 @@ TDefaultErrorResultInternal<TSharedRef<FLobbyDetailsEOS>> FLobbyDetailsEOS::Crea
 	EOS_LobbySearch_CopySearchResultByIndexOptions Options;
 	Options.ApiVersion = EOS_LOBBYSEARCH_COPYSEARCHRESULTBYINDEX_API_LATEST;
 	Options.LobbyIndex = ResultIndex;
+	static_assert(EOS_LOBBYSEARCH_COPYSEARCHRESULTBYINDEX_API_LATEST == 1, "EOS_LobbySearch_CopySearchResultByIndexOptions updated, check new fields");
 
 	EOS_EResult EOSResult = EOS_LobbySearch_CopySearchResultByIndex(SearchHandle, &Options, &LobbyDetailsHandle);
 	if (EOSResult != EOS_EResult::EOS_Success)
@@ -269,6 +274,8 @@ TFuture<TDefaultErrorResultInternal<TSharedRef<FClientLobbySnapshot>>> FLobbyDet
 
 	EOS_LobbyDetails_GetMemberCountOptions GetMemberCountOptions = {};
 	GetMemberCountOptions.ApiVersion = EOS_LOBBYDETAILS_GETMEMBERCOUNT_API_LATEST;
+	static_assert(EOS_LOBBYDETAILS_GETMEMBERCOUNT_API_LATEST == 1, "EOS_LobbyDetails_GetMemberCountOptions updated, check new fields");
+
 	const uint32_t MemberCount = EOS_LobbyDetails_GetMemberCount(LobbyDetailsHandle, &GetMemberCountOptions);
 
 	TSharedRef<TArray<EOS_ProductUserId>> MemberProductUserIds = MakeShared<TArray<EOS_ProductUserId>>();
@@ -279,6 +286,8 @@ TFuture<TDefaultErrorResultInternal<TSharedRef<FClientLobbySnapshot>>> FLobbyDet
 		EOS_LobbyDetails_GetMemberByIndexOptions GetMemberByIndexOptions = {};
 		GetMemberByIndexOptions.ApiVersion = EOS_LOBBYDETAILS_GETMEMBERBYINDEX_API_LATEST;
 		GetMemberByIndexOptions.MemberIndex = MemberIndex;
+		static_assert(EOS_LOBBYDETAILS_GETMEMBERBYINDEX_API_LATEST == 1, "EOS_LobbyDetails_GetMemberByIndexOptions updated, check new fields");
+
 		MemberProductUserIds->Emplace(EOS_LobbyDetails_GetMemberByIndex(LobbyDetailsHandle, &GetMemberByIndexOptions));
 	}
 
@@ -305,6 +314,8 @@ TFuture<TDefaultErrorResultInternal<TSharedRef<FClientLobbySnapshot>>> FLobbyDet
 		{
 			EOS_LobbyDetails_GetLobbyOwnerOptions GetLobbyOwnerOptions = {};
 			GetLobbyOwnerOptions.ApiVersion = EOS_LOBBYDETAILS_GETLOBBYOWNER_API_LATEST;
+			static_assert(EOS_LOBBYDETAILS_GETLOBBYOWNER_API_LATEST == 1, "EOS_LobbyDetails_GetLobbyOwnerOptions updated, check new fields");
+
 			const EOS_ProductUserId LobbyOwner = EOS_LobbyDetails_GetLobbyOwner(StrongThis->LobbyDetailsHandle, &GetLobbyOwnerOptions);
 
 			for (int32 MemberIndex = 0; MemberIndex < MemberProductUserIds->Num(); ++MemberIndex)
@@ -325,6 +336,7 @@ TFuture<TDefaultErrorResultInternal<TSharedRef<FClientLobbySnapshot>>> FLobbyDet
 		{
 			EOS_LobbyDetails_GetAttributeCountOptions GetAttributeCountOptions = {};
 			GetAttributeCountOptions.ApiVersion = EOS_LOBBYDETAILS_GETATTRIBUTECOUNT_API_LATEST;
+			static_assert(EOS_LOBBYDETAILS_GETATTRIBUTECOUNT_API_LATEST == 1, "EOS_LobbyDetails_GetAttributeCountOptions updated, check new fields");
 
 			const uint32_t AttributeCount = EOS_LobbyDetails_GetAttributeCount(StrongThis->LobbyDetailsHandle, &GetAttributeCountOptions);
 			for (uint32_t AttributeIndex = 0; AttributeIndex < AttributeCount; ++AttributeIndex)
@@ -332,6 +344,7 @@ TFuture<TDefaultErrorResultInternal<TSharedRef<FClientLobbySnapshot>>> FLobbyDet
 				EOS_LobbyDetails_CopyAttributeByIndexOptions CopyAttributeByIndexOptions = {};
 				CopyAttributeByIndexOptions.ApiVersion = EOS_LOBBYDETAILS_COPYATTRIBUTEBYINDEX_API_LATEST;
 				CopyAttributeByIndexOptions.AttrIndex = AttributeIndex;
+				static_assert(EOS_LOBBYDETAILS_COPYATTRIBUTEBYINDEX_API_LATEST == 1, "EOS_LobbyDetails_CopyAttributeByIndexOptions updated, check new fields");
 
 				EOS_Lobby_Attribute* LobbyAttribute = nullptr;
 				ON_SCOPE_EXIT
@@ -373,6 +386,7 @@ TDefaultErrorResultInternal<TSharedRef<FClientLobbyMemberSnapshot>> FLobbyDetail
 		EOS_LobbyDetails_GetMemberAttributeCountOptions GetMemberAttributeCountOptions = {};
 		GetMemberAttributeCountOptions.ApiVersion = EOS_LOBBYDETAILS_GETMEMBERATTRIBUTECOUNT_API_LATEST;
 		GetMemberAttributeCountOptions.TargetUserId = MemberProductUserId;
+		static_assert(EOS_LOBBYDETAILS_GETMEMBERATTRIBUTECOUNT_API_LATEST == 1, "EOS_LobbyDetails_GetMemberAttributeCountOptions updated, check new fields");
 
 		const uint32_t MemberAttributeCount = EOS_LobbyDetails_GetMemberAttributeCount(LobbyDetailsHandle, &GetMemberAttributeCountOptions);
 		for (uint32_t MemberAttributeIndex = 0; MemberAttributeIndex < MemberAttributeCount; ++MemberAttributeIndex)
@@ -381,6 +395,7 @@ TDefaultErrorResultInternal<TSharedRef<FClientLobbyMemberSnapshot>> FLobbyDetail
 			CopyMemberAttributeByIndexOptions.ApiVersion = EOS_LOBBYDETAILS_COPYMEMBERATTRIBUTEBYINDEX_API_LATEST;
 			CopyMemberAttributeByIndexOptions.TargetUserId = MemberProductUserId;
 			CopyMemberAttributeByIndexOptions.AttrIndex = MemberAttributeIndex;
+			static_assert(EOS_LOBBYDETAILS_COPYMEMBERATTRIBUTEBYINDEX_API_LATEST == 1, "EOS_LobbyDetails_CopyMemberAttributeByIndexOptions updated, check new fields");
 
 			EOS_Lobby_Attribute* LobbyAttribute = nullptr;
 			ON_SCOPE_EXIT
@@ -418,6 +433,7 @@ TFuture<EOS_EResult> FLobbyDetailsEOS::ApplyLobbyDataUpdateFromLocalChanges(
 	ModificationOptions.ApiVersion = EOS_LOBBY_UPDATELOBBYMODIFICATION_API_LATEST;
 	ModificationOptions.LocalUserId = GetProductUserIdChecked(LocalUserId);
 	ModificationOptions.LobbyId = GetInfo()->GetLobbyId();
+	static_assert(EOS_LOBBY_UPDATELOBBYMODIFICATION_API_LATEST == 1, "EOS_Lobby_UpdateLobbyModificationOptions updated, check new fields");
 
 	EOS_EResult EOSResultCode = EOS_Lobby_UpdateLobbyModification(Prerequisites->LobbyInterfaceHandle, &ModificationOptions, &LobbyModificationHandle);
 	if (EOSResultCode != EOS_EResult::EOS_Success)
@@ -432,6 +448,7 @@ TFuture<EOS_EResult> FLobbyDetailsEOS::ApplyLobbyDataUpdateFromLocalChanges(
 		EOS_LobbyModification_SetPermissionLevelOptions SetPermissionOptions = {};
 		SetPermissionOptions.ApiVersion = EOS_LOBBYMODIFICATION_SETPERMISSIONLEVEL_API_LATEST;
 		SetPermissionOptions.PermissionLevel = TranslateJoinPolicy(*Changes.JoinPolicy);
+		static_assert(EOS_LOBBYMODIFICATION_SETPERMISSIONLEVEL_API_LATEST == 1, "EOS_LobbyModification_SetPermissionLevelOptions updated, check new fields");
 
 		EOSResultCode = EOS_LobbyModification_SetPermissionLevel(LobbyModificationHandle, &SetPermissionOptions);
 		if (EOSResultCode != EOS_EResult::EOS_Success)
@@ -450,6 +467,8 @@ TFuture<EOS_EResult> FLobbyDetailsEOS::ApplyLobbyDataUpdateFromLocalChanges(
 		AddAttributeOptions.ApiVersion = EOS_LOBBYMODIFICATION_ADDATTRIBUTE_API_LATEST;
 		AddAttributeOptions.Attribute = &AttributeTranslator.GetAttributeData();
 		AddAttributeOptions.Visibility = EOS_ELobbyAttributeVisibility::EOS_LAT_PUBLIC; // todo - get from schema
+		static_assert(EOS_LOBBYMODIFICATION_ADDATTRIBUTE_API_LATEST == 1, "EOS_LobbyModification_AddAttributeOptions updated, check new fields");
+
 		EOSResultCode = EOS_LobbyModification_AddAttribute(LobbyModificationHandle, &AddAttributeOptions);
 		if (EOSResultCode != EOS_EResult::EOS_Success)
 		{
@@ -466,6 +485,8 @@ TFuture<EOS_EResult> FLobbyDetailsEOS::ApplyLobbyDataUpdateFromLocalChanges(
 		EOS_LobbyModification_RemoveAttributeOptions RemoveAttributeOptions = {};
 		RemoveAttributeOptions.ApiVersion = EOS_LOBBYMODIFICATION_REMOVEATTRIBUTE_API_LATEST;
 		RemoveAttributeOptions.Key = KeyConverter.Get();
+		static_assert(EOS_LOBBYMODIFICATION_REMOVEATTRIBUTE_API_LATEST == 1, "EOS_LobbyModification_RemoveAttributeOptions updated, check new fields");
+		
 		EOSResultCode = EOS_LobbyModification_RemoveAttribute(LobbyModificationHandle, &RemoveAttributeOptions);
 		if (EOSResultCode != EOS_EResult::EOS_Success)
 		{
@@ -481,6 +502,7 @@ TFuture<EOS_EResult> FLobbyDetailsEOS::ApplyLobbyDataUpdateFromLocalChanges(
 	EOS_Lobby_UpdateLobbyOptions UpdateLobbyOptions = {};
 	UpdateLobbyOptions.ApiVersion = EOS_LOBBY_UPDATELOBBY_API_LATEST;
 	UpdateLobbyOptions.LobbyModificationHandle = LobbyModificationHandle;
+	static_assert(EOS_LOBBY_UPDATELOBBY_API_LATEST == 1, "EOS_Lobby_UpdateLobbyOptions updated, check new fields");
 
 	EOS_Async(EOS_Lobby_UpdateLobby, Prerequisites->LobbyInterfaceHandle, UpdateLobbyOptions,
 	[Promise = MoveTemp(Promise)](const EOS_Lobby_UpdateLobbyCallbackInfo* CallbackInfo) mutable
@@ -506,6 +528,7 @@ TFuture<EOS_EResult> FLobbyDetailsEOS::ApplyLobbyMemberDataUpdateFromLocalChange
 	ModificationOptions.ApiVersion = EOS_LOBBY_UPDATELOBBYMODIFICATION_API_LATEST;
 	ModificationOptions.LocalUserId = GetProductUserIdChecked(LocalUserId);
 	ModificationOptions.LobbyId = GetInfo()->GetLobbyId();
+	static_assert(EOS_LOBBY_UPDATELOBBYMODIFICATION_API_LATEST == 1, "EOS_Lobby_UpdateLobbyModificationOptions updated, check new fields");
 
 	EOS_EResult EOSResultCode = EOS_Lobby_UpdateLobbyModification(Prerequisites->LobbyInterfaceHandle, &ModificationOptions, &LobbyModificationHandle);
 	if (EOSResultCode != EOS_EResult::EOS_Success)
@@ -523,6 +546,8 @@ TFuture<EOS_EResult> FLobbyDetailsEOS::ApplyLobbyMemberDataUpdateFromLocalChange
 		AddMemberAttributeOptions.ApiVersion = EOS_LOBBYMODIFICATION_ADDMEMBERATTRIBUTE_API_LATEST;
 		AddMemberAttributeOptions.Attribute = &AttributeTranslator.GetAttributeData();
 		AddMemberAttributeOptions.Visibility = EOS_ELobbyAttributeVisibility::EOS_LAT_PUBLIC; // todo - get from schema
+		static_assert(EOS_LOBBYMODIFICATION_ADDMEMBERATTRIBUTE_API_LATEST == 1, "EOS_LobbyModification_AddMemberAttributeOptions updated, check new fields");
+		
 		EOSResultCode = EOS_LobbyModification_AddMemberAttribute(LobbyModificationHandle, &AddMemberAttributeOptions);
 		if (EOSResultCode != EOS_EResult::EOS_Success)
 		{
@@ -539,6 +564,8 @@ TFuture<EOS_EResult> FLobbyDetailsEOS::ApplyLobbyMemberDataUpdateFromLocalChange
 		EOS_LobbyModification_RemoveMemberAttributeOptions RemoveMemberAttributeOptions = {};
 		RemoveMemberAttributeOptions.ApiVersion = EOS_LOBBYMODIFICATION_REMOVEMEMBERATTRIBUTE_API_LATEST;
 		RemoveMemberAttributeOptions.Key = KeyConverter.Get();
+		static_assert(EOS_LOBBYMODIFICATION_REMOVEMEMBERATTRIBUTE_API_LATEST == 1, "EOS_LobbyModification_RemoveMemberAttributeOptions updated, check new fields");
+		
 		EOSResultCode = EOS_LobbyModification_RemoveMemberAttribute(LobbyModificationHandle, &RemoveMemberAttributeOptions);
 		if (EOSResultCode != EOS_EResult::EOS_Success)
 		{
@@ -551,6 +578,7 @@ TFuture<EOS_EResult> FLobbyDetailsEOS::ApplyLobbyMemberDataUpdateFromLocalChange
 	EOS_Lobby_UpdateLobbyOptions UpdateLobbyOptions = {};
 	UpdateLobbyOptions.ApiVersion = EOS_LOBBY_UPDATELOBBY_API_LATEST;
 	UpdateLobbyOptions.LobbyModificationHandle = LobbyModificationHandle;
+	static_assert(EOS_LOBBY_UPDATELOBBY_API_LATEST == 1, "EOS_Lobby_UpdateLobbyOptions updated, check new fields");
 
 	TPromise<EOS_EResult> Promise;
 	TFuture<EOS_EResult> Future = Promise.GetFuture();
@@ -582,6 +610,7 @@ TDefaultErrorResultInternal<TSharedRef<FLobbyDetailsInfoEOS>> FLobbyDetailsInfoE
 {
 	EOS_LobbyDetails_CopyInfoOptions CopyInfoOptions = {};
 	CopyInfoOptions.ApiVersion = EOS_LOBBYDETAILS_COPYINFO_API_LATEST;
+	static_assert(EOS_LOBBYDETAILS_COPYINFO_API_LATEST == 1, "EOS_LobbyDetails_CopyInfoOptions updated, check new fields");
 
 	EOS_LobbyDetails_Info* LobbyDetailsInfo = nullptr;
 	EOS_EResult EOSResult = EOS_LobbyDetails_CopyInfo(LobbyDetailsHandle, &CopyInfoOptions, &LobbyDetailsInfo);
@@ -854,6 +883,7 @@ TFuture<TDefaultErrorResultInternal<TSharedRef<FLobbySearchEOS>>> FLobbySearchEO
 	EOS_Lobby_CreateLobbySearchOptions CreateLobbySearchOptions = {};
 	CreateLobbySearchOptions.ApiVersion = EOS_LOBBY_CREATELOBBYSEARCH_API_LATEST;
 	CreateLobbySearchOptions.MaxResults = Params.MaxResults;
+	static_assert(EOS_LOBBY_CREATELOBBYSEARCH_API_LATEST == 1, "EOS_Lobby_CreateLobbySearchOptions updated, check new fields");
 
 	EOS_EResult EOSResult = EOS_Lobby_CreateLobbySearch(Prerequisites->LobbyInterfaceHandle, &CreateLobbySearchOptions, &SearchHandle->Get());
 	if (EOSResult != EOS_EResult::EOS_Success)
@@ -873,6 +903,7 @@ TFuture<TDefaultErrorResultInternal<TSharedRef<FLobbySearchEOS>>> FLobbySearchEO
 		EOS_LobbySearch_SetLobbyIdOptions SetLobbyIdOptions = {};
 		SetLobbyIdOptions.ApiVersion = EOS_LOBBYSEARCH_SETLOBBYID_API_LATEST;
 		SetLobbyIdOptions.LobbyId = LobbyData->GetLobbyIdEOS();
+		static_assert(EOS_LOBBYSEARCH_SETLOBBYID_API_LATEST == 1, "EOS_LobbySearch_SetLobbyIdOptions updated, check new fields");
 
 		EOSResult = EOS_LobbySearch_SetLobbyId(SearchHandle->Get(), &SetLobbyIdOptions);
 		if (EOSResult != EOS_EResult::EOS_Success)
@@ -886,6 +917,7 @@ TFuture<TDefaultErrorResultInternal<TSharedRef<FLobbySearchEOS>>> FLobbySearchEO
 		EOS_LobbySearch_SetTargetUserIdOptions SetTargetUserIdOptions = {};
 		SetTargetUserIdOptions.ApiVersion = EOS_LOBBYSEARCH_SETTARGETUSERID_API_LATEST;
 		SetTargetUserIdOptions.TargetUserId = GetProductUserIdChecked(*Params.TargetUser);
+		static_assert(EOS_LOBBYSEARCH_SETTARGETUSERID_API_LATEST == 1, "EOS_LobbySearch_SetTargetUserIdOptions updated, check new fields");
 
 		EOSResult = EOS_LobbySearch_SetTargetUserId(SearchHandle->Get(), &SetTargetUserIdOptions);
 		if (EOSResult != EOS_EResult::EOS_Success)
@@ -905,11 +937,13 @@ TFuture<TDefaultErrorResultInternal<TSharedRef<FLobbySearchEOS>>> FLobbySearchEO
 			AttributeData.Key = EOS_LOBBY_SEARCH_BUCKET_ID;
 			AttributeData.ValueType = EOS_ELobbyAttributeType::EOS_AT_STRING;
 			AttributeData.Value.AsUtf8 = BucketTranslator.GetBucketIdEOS();
+			static_assert(EOS_LOBBY_ATTRIBUTEDATA_API_LATEST == 1, "EOS_Lobby_AttributeData updated, check new fields");
 
 			EOS_LobbySearch_SetParameterOptions SetParameterOptions = {};
-			SetParameterOptions.ApiVersion = EOS_LOBBYSEARCH_SETTARGETUSERID_API_LATEST;
+			SetParameterOptions.ApiVersion = EOS_LOBBYSEARCH_SETPARAMETER_API_LATEST;
 			SetParameterOptions.Parameter = &AttributeData;
 			SetParameterOptions.ComparisonOp = EOS_EComparisonOp::EOS_CO_EQUAL;
+			static_assert(EOS_LOBBYSEARCH_SETPARAMETER_API_LATEST == 1, "EOS_LobbySearch_SetParameterOptions updated, check new fields");
 
 			EOSResult = EOS_LobbySearch_SetParameter(SearchHandle->Get(), &SetParameterOptions);
 			if (EOSResult != EOS_EResult::EOS_Success)
@@ -925,9 +959,10 @@ TFuture<TDefaultErrorResultInternal<TSharedRef<FLobbySearchEOS>>> FLobbySearchEO
 				Filter.AttributeName, Filter.ComparisonValue);
 
 			EOS_LobbySearch_SetParameterOptions SetParameterOptions = {};
-			SetParameterOptions.ApiVersion = EOS_LOBBYSEARCH_SETTARGETUSERID_API_LATEST;
+			SetParameterOptions.ApiVersion = EOS_LOBBYSEARCH_SETPARAMETER_API_LATEST;
 			SetParameterOptions.Parameter = &AttributeTranslator.GetAttributeData();
 			SetParameterOptions.ComparisonOp = Private::TranslateSearchComparison(Filter.ComparisonOp);
+			static_assert(EOS_LOBBYSEARCH_SETPARAMETER_API_LATEST == 1, "EOS_LobbySearch_SetParameterOptions updated, check new fields");
 
 			EOSResult = EOS_LobbySearch_SetParameter(SearchHandle->Get(), &SetParameterOptions);
 			if (EOSResult != EOS_EResult::EOS_Success)
@@ -944,6 +979,7 @@ TFuture<TDefaultErrorResultInternal<TSharedRef<FLobbySearchEOS>>> FLobbySearchEO
 	EOS_LobbySearch_FindOptions FindOptions = {};
 	FindOptions.ApiVersion = EOS_LOBBYSEARCH_FIND_API_LATEST;
 	FindOptions.LocalUserId = GetProductUserIdChecked(Params.LocalUserId);
+	static_assert(EOS_LOBBYSEARCH_FIND_API_LATEST == 1, "EOS_LobbySearch_FindOptions updated, check new fields");
 
 	EOS_Async(EOS_LobbySearch_Find, SearchHandle->Get(), FindOptions,
 	[Promise = MoveTemp(Promise), Prerequisites, LobbyRegistry, LocalUserId = Params.LocalUserId, SearchHandle]
@@ -960,6 +996,8 @@ TFuture<TDefaultErrorResultInternal<TSharedRef<FLobbySearchEOS>>> FLobbySearchEO
 
 		EOS_LobbySearch_GetSearchResultCountOptions GetSearchResultCountOptions = {};
 		GetSearchResultCountOptions.ApiVersion = EOS_LOBBYSEARCH_GETSEARCHRESULTCOUNT_API_LATEST;
+		static_assert(EOS_LOBBYSEARCH_GETSEARCHRESULTCOUNT_API_LATEST == 1, "EOS_LobbySearch_GetSearchResultCountOptions updated, check new fields");
+		
 		const uint32_t NumSearchResults = EOS_LobbySearch_GetSearchResultCount(SearchHandle->Get(), &GetSearchResultCountOptions);
 
 		for (uint32_t SearchResultIndex = 0; SearchResultIndex < NumSearchResults; ++SearchResultIndex)
