- Fix EOSGS Achievements not updating the Progress field, so Unlock could never return AlreadyUnlocked

diff --git a/Engine/Plugins/Online/Experimental/OnlineServicesEOSGS/Source/Private/Online/AchievementsEOSGS.cpp b/Engine/Plugins/Online/Experimental/OnlineServicesEOSGS/Source/Private/Online/AchievementsEOSGS.cpp
--- a/Engine/Plugins/Online/Experimental/OnlineServicesEOSGS/Source/Private/Online/AchievementsEOSGS.cpp
+++ b/Engine/Plugins/Online/Experimental/OnlineServicesEOSGS/Source/Private/Online/AchievementsEOSGS.cpp
@@ -33,12 +33,26 @@ void FAchievementsEOSGS::Initialize()
 	{
 		FAchievementsEOSGS* This = reinterpret_cast<FAchievementsEOSGS*>(Data->ClientData);
 
-		FOnlineAccountIdHandle LocalUserId = FindAccountId(Data->UserId);
+		const FOnlineAccountIdHandle LocalUserId = FindAccountId(Data->UserId);
 		if (LocalUserId.IsValid())
 		{
+			const FString AchievementId = UTF8_TO_TCHAR(Data->AchievementId);
+			const FDateTime UnlockTime = FDateTime::FromUnixTimestamp(Data->UnlockTime);
+
+			UE_LOG(LogTemp, Verbose, TEXT("EOS_Achievements_OnAchievementsUnlocked User=%s AchievId=%s"), *ToLogString(LocalUserId), *AchievementId);
+
+			if (FAchievementStateMap* LocalUserAchievementStates = This->AchievementStates.Find(LocalUserId))
+			{
+				if (FAchievementState* AchievementState = LocalUserAchievementStates->Find(AchievementId))
+				{
+					AchievementState->Progress = 1.0f;
+					AchievementState->UnlockTime = UnlockTime;
+				}
+			}
+			
 			FAchievementStateUpdated EventParams;
 			EventParams.LocalUserId = LocalUserId;
-			EventParams.AchievementIds.Emplace(UTF8_TO_TCHAR(Data->AchievementId));
+			EventParams.AchievementIds.Emplace(AchievementId);
 
 			This->OnAchievementStateUpdatedEvent.Broadcast(EventParams);
 		}
@@ -373,17 +387,34 @@ TOnlineAsyncOpHandle<FUnlockAchievements> FAchievementsEOSGS::UnlockAchievements
 	})
 	.Then([this](TOnlineAsyncOp<FUnlockAchievements>& InAsyncOp, const EOS_Achievements_OnUnlockAchievementsCompleteCallbackInfo* Data)
 	{
-		if (Data->ResultCode != EOS_EResult::EOS_Success)
+		if (Data->ResultCode == EOS_EResult::EOS_Success)
 		{
-			UE_LOG(LogTemp, Warning, TEXT("EOS_Achievements_UnlockAchievements failed with result=[%s]"), *LexToString(Data->ResultCode));
+			UE_LOG(LogTemp, Verbose, TEXT("EOS_Achievements_UnlockAchievements succeeded"), *LexToString(Data->ResultCode));
 
-			InAsyncOp.SetError(Errors::FromEOSResult(Data->ResultCode));
+			const FUnlockAchievements::Params& Params = InAsyncOp.GetParams();
+			const FDateTime TimeNow = FDateTime::Now();
+
+			FAchievementStateMap* LocalUserAchievementStates = AchievementStates.Find(Params.LocalUserId);
+			if (ensure(LocalUserAchievementStates))
+			{
+				for(const FString& AchievementId : Params.AchievementIds)
+				{
+					FAchievementState* AchievementState = LocalUserAchievementStates->Find(AchievementId);
+					if (ensure(AchievementState))
+					{
+						AchievementState->Progress = 1.0f;
+						AchievementState->UnlockTime = TimeNow;
+					}
+				}
+			}
+
+			InAsyncOp.SetResult({});
 		}
 		else
 		{
-			InAsyncOp.SetResult({});
+			UE_LOG(LogTemp, Warning, TEXT("EOS_Achievements_UnlockAchievements failed with result=[%s]"), *LexToString(Data->ResultCode));
+			InAsyncOp.SetError(Errors::FromEOSResult(Data->ResultCode));
 		}
-
 	})
 	.Enqueue(GetSerialQueue());
 
