- Add encryption key to EOSGS platform options
- Fix logging warnings about socket subsystem when eos platform creation fails
- Fix creating platforms in component lifetime functions (Initialize/Shutdown), not service lifetime functions (Init/Destroy)
- Fix not creating TitleFile / UserFile components.
- Fix TitleFile::ReadFile being totally broken, and crashing in the ReadFileData callback

diff --git a/Engine/Plugins/Online/Experimental/OnlineServicesEOSGS/Source/Private/Online/OnlineServicesEOSGS.cpp b/Engine/Plugins/Online/Experimental/OnlineServicesEOSGS/Source/Private/Online/OnlineServicesEOSGS.cpp
--- a/Engine/Plugins/Online/Experimental/OnlineServicesEOSGS/Source/Private/Online/OnlineServicesEOSGS.cpp
+++ b/Engine/Plugins/Online/Experimental/OnlineServicesEOSGS/Source/Private/Online/OnlineServicesEOSGS.cpp
@@ -6,11 +6,13 @@
 #include "Online/AuthEOSGS.h"
 #include "Online/LeaderboardsEOSGS.h"
 #include "Online/LobbiesEOSGS.h"
-#include "Online/StatsEOSGS.h"
-#include "Online/SessionsEOSGS.h"
 #include "Online/OnlineIdEOSGS.h"
 #include "Online/OnlineServicesEOSGSTypes.h"
 #include "Online/OnlineServicesEOSGSPlatformFactory.h"
+#include "Online/StatsEOSGS.h"
+#include "Online/SessionsEOSGS.h"
+#include "Online/TitleFileEOSGS.h"
+#include "Online/UserFileEOSGS.h"
 #include "IEOSSDKManager.h"
 
 #if WITH_ENGINE
@@ -26,18 +28,7 @@ FOnlineServicesEOSGS::FOnlineServicesEOSGS(FName InInstanceName)
 {
 }
 
-void FOnlineServicesEOSGS::RegisterComponents()
-{
-	Components.Register<FAchievementsEOSGS>(*this);
-	Components.Register<FAuthEOSGS>(*this);
-	Components.Register<FLeaderboardsEOSGS>(*this);
-	Components.Register<FLobbiesEOSGS>(*this);
-	Components.Register<FStatsEOSGS>(*this);
-	Components.Register<FSessionsEOSGS>(*this);
-	Super::RegisterComponents();
-}
-
-void FOnlineServicesEOSGS::Initialize()
+void FOnlineServicesEOSGS::Init()
 {
 	FOnlineServicesEOSGSPlatformFactory& PlatformFactory = FOnlineServicesEOSGSPlatformFactory::Get();
 	if (InstanceName.IsNone())
@@ -67,12 +58,12 @@ void FOnlineServicesEOSGS::Initialize()
 		return;
 	}
 
-	Super::Initialize();
+	Super::Init();
 }
 
-void FOnlineServicesEOSGS::Shutdown()
+void FOnlineServicesEOSGS::Destroy()
 {
-	Super::Shutdown();
+	Super::Destroy();
 #if WITH_ENGINE
 	if (SocketSubsystem)
 	{
@@ -82,6 +73,25 @@ void FOnlineServicesEOSGS::Shutdown()
 #endif
 }
 
+void FOnlineServicesEOSGS::RegisterComponents()
+{
+	Components.Register<FAchievementsEOSGS>(*this);
+	Components.Register<FAuthEOSGS>(*this);
+	Components.Register<FLeaderboardsEOSGS>(*this);
+	Components.Register<FLobbiesEOSGS>(*this);
+	Components.Register<FStatsEOSGS>(*this);
+	Components.Register<FSessionsEOSGS>(*this);
+	if (EOS_Platform_GetTitleStorageInterface(GetEOSPlatformHandle()))
+	{
+		Components.Register<FTitleFileEOSGS>(*this);
+	}
+	if (EOS_Platform_GetPlayerDataStorageInterface(GetEOSPlatformHandle()))
+	{
+		Components.Register<FUserFileEOSGS>(*this);
+	}
+	Super::RegisterComponents();
+}
+
 TOnlineResult<FGetResolvedConnectString> FOnlineServicesEOSGS::GetResolvedConnectString(FGetResolvedConnectString::Params&& Params)
 {
 	if (Params.LobbyId.IsValid())
diff --git a/Engine/Plugins/Online/Experimental/OnlineServicesEOSGS/Source/Private/Online/OnlineServicesEOSGSPlatformFactory.cpp b/Engine/Plugins/Online/Experimental/OnlineServicesEOSGS/Source/Private/Online/OnlineServicesEOSGSPlatformFactory.cpp
--- a/Engine/Plugins/Online/Experimental/OnlineServicesEOSGS/Source/Private/Online/OnlineServicesEOSGSPlatformFactory.cpp
+++ b/Engine/Plugins/Online/Experimental/OnlineServicesEOSGS/Source/Private/Online/OnlineServicesEOSGSPlatformFactory.cpp
@@ -21,6 +21,7 @@ struct FEOSPlatformConfig
 	FString DeploymentId;
 	FString ClientId;
 	FString ClientSecret;
+	FString EncryptionKey;
 };
 
 FEOSPlatformConfig LoadEOSPlatformConfig()
@@ -32,6 +33,7 @@ FEOSPlatformConfig LoadEOSPlatformConfig()
 	GConfig->GetString(*ConfigSection, TEXT("DeploymentId"), PlatformConfig.DeploymentId, GEngineIni);
 	GConfig->GetString(*ConfigSection, TEXT("ClientId"), PlatformConfig.ClientId, GEngineIni);
 	GConfig->GetString(*ConfigSection, TEXT("ClientSecret"), PlatformConfig.ClientSecret, GEngineIni);
+	GConfig->GetString(*ConfigSection, TEXT("EncryptionKey"), PlatformConfig.EncryptionKey, GEngineIni);
 	return PlatformConfig;
 }
 
@@ -91,6 +93,7 @@ IEOSPlatformHandlePtr FOnlineServicesEOSGSPlatformFactory::CreatePlatform()
 	const FTCHARToUTF8 DeploymentId(*EOSPlatformConfig.DeploymentId);
 	const FTCHARToUTF8 ClientId(*EOSPlatformConfig.ClientId);
 	const FTCHARToUTF8 ClientSecret(*EOSPlatformConfig.ClientSecret);
+	const FTCHARToUTF8 EncryptionKey(EOSPlatformConfig.EncryptionKey.IsEmpty() ? nullptr : *EOSPlatformConfig.EncryptionKey);
 	const FTCHARToUTF8 CacheDirectory(*(SDKManager->GetCacheDirBase() / TEXT("OnlineServicesEOS")));
 
 	EOS_Platform_Options PlatformOptions = {};
@@ -114,13 +117,13 @@ IEOSPlatformHandlePtr FOnlineServicesEOSGSPlatformFactory::CreatePlatform()
 	PlatformOptions.DeploymentId = DeploymentId.Get();
 	PlatformOptions.ClientCredentials.ClientId = ClientId.Get();
 	PlatformOptions.ClientCredentials.ClientSecret = ClientSecret.Get();
+	PlatformOptions.EncryptionKey = EncryptionKey.Get();
 	PlatformOptions.CacheDirectory = CacheDirectory.Get();
 
 	IEOSPlatformHandlePtr EOSPlatformHandle = SDKManager->CreatePlatform(PlatformOptions);
 	if (!EOSPlatformHandle)
 	{
-		UE_LOG(LogOnlineServices, Verbose, TEXT("[FOnlineServicesEOSGS::Initialize] Unable to initialize Socket Subsystem. EOS Platform Handle was invalid."));
-		return {};
+		UE_LOG(LogOnlineServices, Warning, TEXT("[FOnlineServicesEOSGS::Initialize] Failed to create platform."));
 	}
 	return EOSPlatformHandle;
 }
diff --git a/Engine/Plugins/Online/Experimental/OnlineServicesEOSGS/Source/Private/Online/TitleFileEOSGS.cpp b/Engine/Plugins/Online/Experimental/OnlineServicesEOSGS/Source/Private/Online/TitleFileEOSGS.cpp
--- a/Engine/Plugins/Online/Experimental/OnlineServicesEOSGS/Source/Private/Online/TitleFileEOSGS.cpp
+++ b/Engine/Plugins/Online/Experimental/OnlineServicesEOSGS/Source/Private/Online/TitleFileEOSGS.cpp
@@ -136,6 +136,19 @@ TOnlineResult<FTitleFileGetEnumeratedFiles> FTitleFileEOSGS::GetEnumeratedFiles(
 	return TOnlineResult<FTitleFileGetEnumeratedFiles>({*EnumeratedFiles});
 }
 
+struct FTitleFileReadFileClientData
+{
+	using PromiseType = TPromise<const EOS_TitleStorage_ReadFileCallbackInfo*>;
+
+	FTitleFileReadFileClientData(const FTitleFileContentsRef& InFileContents, PromiseType&& InPromise)
+		: FileContents(InFileContents)
+		, Promise(MoveTemp(InPromise))
+	{}
+
+	FTitleFileContentsWeakPtr FileContents;
+	PromiseType Promise;
+};
+
 TOnlineAsyncOpHandle<FTitleFileReadFile> FTitleFileEOSGS::ReadFile(FTitleFileReadFile::Params&& InParams)
 {
 	static const TCHAR* FileContentsKey = TEXT("FileContents");
@@ -164,22 +177,7 @@ TOnlineAsyncOpHandle<FTitleFileReadFile> FTitleFileEOSGS::ReadFile(FTitleFileRea
 		FTitleFileContentsRef FileContentsRef = MakeShared<FTitleFileContents>();
 		InAsyncOp.Data.Set<FTitleFileContentsRef>(FileContentsKey, FileContentsRef);
 
-		using FReadFileDataCallback = TEOSGlobalCallback<EOS_TitleStorage_OnReadFileDataCallback, EOS_TitleStorage_ReadFileDataCallbackInfo, TOnlineAsyncOp<FTitleFileReadFile>, EOS_TitleStorage_EReadResult>;
-		TSharedRef<FReadFileDataCallback> ReadFileDataCallback = MakeShared<FReadFileDataCallback>(InAsyncOp.AsWeak());
-		InAsyncOp.Data.Set<TSharedRef<FReadFileDataCallback>>(TEXT("ReadFileDataCallback"), ReadFileDataCallback);
-		ReadFileDataCallback->CallbackLambda = [FileContentsRef](const EOS_TitleStorage_ReadFileDataCallbackInfo* Data) -> EOS_TitleStorage_EReadResult
-		{
-			FTitleFileContents& FileContents = const_cast<FTitleFileContents&>(FileContentsRef.Get());
-			// Only has any effect on the first callback
-			FileContents.Reserve(Data->TotalFileSizeBytes);
-			FileContents.Append((uint8*)Data->DataChunk, Data->DataChunkLengthBytes);
-
-			UE_LOG(LogTemp, VeryVerbose, TEXT("FTitleFileEOSGS::ReadFile ReadProgress Filename=[%s] %d/%d"), UTF8_TO_TCHAR(Data->Filename), FileContents.Num(), Data->TotalFileSizeBytes);
-
-			return EOS_TitleStorage_EReadResult::EOS_TS_RR_ContinueReading;
-		};
-
-		const FTCHARToUTF8 Utf8Filename(Params.Filename);
+		const FTCHARToUTF8 Utf8Filename(*Params.Filename);
 
 		EOS_TitleStorage_ReadFileOptions Options = {};
 		Options.ApiVersion = EOS_TITLESTORAGE_READFILEOPTIONS_API_LATEST;
@@ -187,10 +185,12 @@ TOnlineAsyncOpHandle<FTitleFileReadFile> FTitleFileEOSGS::ReadFile(FTitleFileRea
 		Options.LocalUserId = GetProductUserIdChecked(Params.LocalUserId);
 		Options.Filename = Utf8Filename.Get();
 		Options.ReadChunkLengthBytes = Config.ReadChunkLengthBytes;
-		Options.ReadFileDataCallback = ReadFileDataCallback->GetCallbackPtr();
+		Options.ReadFileDataCallback = &FTitleFileEOSGS::OnReadFileDataStatic;
 		Options.FileTransferProgressCallback = &FTitleFileEOSGS::OnFileTransferProgressStatic;
 
-		const EOS_HTitleStorageFileTransferRequest RequestHandle = EOS_Async(EOS_TitleStorage_ReadFile, TitleStorageHandle, Options, MoveTemp(Promise));
+		FTitleFileReadFileClientData* ClientData = new FTitleFileReadFileClientData(FileContentsRef, MoveTemp(Promise));
+
+		const EOS_HTitleStorageFileTransferRequest RequestHandle = EOS_TitleStorage_ReadFile(TitleStorageHandle, &Options, ClientData, &FTitleFileEOSGS::OnReadFileCompleteStatic);
 		InAsyncOp.Data.Set<EOS_HTitleStorageFileTransferRequest>(RequestHandleKey, RequestHandle);
 	})
 	.Then([this](TOnlineAsyncOp<FTitleFileReadFile>& InAsyncOp, const EOS_TitleStorage_ReadFileCallbackInfo* Data)
@@ -222,9 +222,37 @@ TOnlineAsyncOpHandle<FTitleFileReadFile> FTitleFileEOSGS::ReadFile(FTitleFileRea
 	return Op->GetHandle();
 }
 
+EOS_TitleStorage_EReadResult EOS_CALL FTitleFileEOSGS::OnReadFileDataStatic(const EOS_TitleStorage_ReadFileDataCallbackInfo* Data)
+{
+	FTitleFileReadFileClientData* ClientData = (FTitleFileReadFileClientData*)Data->ClientData;
+	check(ClientData);
+	FTitleFileContentsPtr FileContentsPtr = ClientData->FileContents.Pin();
+	if (ensure(FileContentsPtr))
+	{
+		FTitleFileContents& FileContents = *ConstCastSharedPtr<FTitleFileContents>(FileContentsPtr);
+		// Only has any effect on the first callback
+		FileContents.Reserve(Data->TotalFileSizeBytes);
+		FileContents.Append((uint8*)Data->DataChunk, Data->DataChunkLengthBytes);
+
+		UE_LOG(LogTemp, VeryVerbose, TEXT("FTitleFileEOSGS::ReadFile ReadProgress Filename=[%s] %d/%d"), UTF8_TO_TCHAR(Data->Filename), FileContents.Num(), Data->TotalFileSizeBytes);
+	}
+
+	return EOS_TitleStorage_EReadResult::EOS_TS_RR_ContinueReading;
+}
 void EOS_CALL FTitleFileEOSGS::OnFileTransferProgressStatic(const EOS_TitleStorage_FileTransferProgressCallbackInfo* Data)
 {
 	UE_LOG(LogTemp, VeryVerbose, TEXT("FTitleFileEOSGS::ReadFile TransferProgress Filename=[%s] %d/%d"), UTF8_TO_TCHAR(Data->Filename), Data->BytesTransferred, Data->TotalFileSizeBytes);
 }
 
+void EOS_CALL FTitleFileEOSGS::OnReadFileCompleteStatic(const EOS_TitleStorage_ReadFileCallbackInfo* Data)
+{
+	UE_LOG(LogTemp, VeryVerbose, TEXT("FTitleFileEOSGS::ReadFile Complete Filename=[%s]"), UTF8_TO_TCHAR(Data->Filename));
+
+	FTitleFileReadFileClientData* ClientData = (FTitleFileReadFileClientData*)Data->ClientData;
+	check(ClientData);
+	ClientData->Promise.EmplaceValue(Data);
+
+	delete ClientData;
+}
+
 /* UE::Online */ }
diff --git a/Engine/Plugins/Online/Experimental/OnlineServicesEOSGS/Source/Public/Online/OnlineServicesEOSGS.h b/Engine/Plugins/Online/Experimental/OnlineServicesEOSGS/Source/Public/Online/OnlineServicesEOSGS.h
--- a/Engine/Plugins/Online/Experimental/OnlineServicesEOSGS/Source/Public/Online/OnlineServicesEOSGS.h
+++ b/Engine/Plugins/Online/Experimental/OnlineServicesEOSGS/Source/Public/Online/OnlineServicesEOSGS.h
@@ -25,9 +25,9 @@ public:
 	FOnlineServicesEOSGS(FName InInstanceName);
 	virtual ~FOnlineServicesEOSGS() = default;
 
+	virtual void Init() override;
+	virtual void Destroy() override;
 	virtual void RegisterComponents() override;
-	virtual void Initialize() override;
-	virtual void Shutdown() override;
 	virtual TOnlineResult<FGetResolvedConnectString> GetResolvedConnectString(FGetResolvedConnectString::Params&& Params) override;
 	virtual EOnlineServices GetServicesProvider() const override { return EOnlineServices::Epic; }
 
diff --git a/Engine/Plugins/Online/Experimental/OnlineServicesEOSGS/Source/Public/Online/TitleFileEOSGS.h b/Engine/Plugins/Online/Experimental/OnlineServicesEOSGS/Source/Public/Online/TitleFileEOSGS.h
--- a/Engine/Plugins/Online/Experimental/OnlineServicesEOSGS/Source/Public/Online/TitleFileEOSGS.h
+++ b/Engine/Plugins/Online/Experimental/OnlineServicesEOSGS/Source/Public/Online/TitleFileEOSGS.h
@@ -53,7 +53,9 @@ protected:
 
 	TOptional<TArray<FString>> EnumeratedFiles;
 
+	static EOS_TitleStorage_EReadResult EOS_CALL OnReadFileDataStatic(const EOS_TitleStorage_ReadFileDataCallbackInfo* Data);
 	static void EOS_CALL OnFileTransferProgressStatic(const EOS_TitleStorage_FileTransferProgressCallbackInfo* Data);
+	static void EOS_CALL OnReadFileCompleteStatic(const EOS_TitleStorage_ReadFileCallbackInfo* Data);
 };
 
 /* UE::Online */ }
