Added caches for search results and invites in sessions interface
- Refactored methods taking shared references as parameters to take id handles instead
- Added invite id handle and registry types, and refactored invite logic to use them.
- Added new GetSessionInvites method to the interface to make use of invites cache, and defined common behaviour.

diff --git a/Engine/Plugins/Online/Experimental/OnlineServices/Source/OnlineServicesCommon/Private/Online/SessionsCommon.cpp b/Engine/Plugins/Online/Experimental/OnlineServices/Source/OnlineServicesCommon/Private/Online/SessionsCommon.cpp
--- a/Engine/Plugins/Online/Experimental/OnlineServices/Source/OnlineServicesCommon/Private/Online/SessionsCommon.cpp
+++ b/Engine/Plugins/Online/Experimental/OnlineServices/Source/OnlineServicesCommon/Private/Online/SessionsCommon.cpp
@@ -30,9 +30,9 @@ namespace UE::Online {
 		RegisterCommand(&FSessionsCommon::LeaveSession);
 		RegisterCommand(&FSessionsCommon::FindSessions);
 		RegisterCommand(&FSessionsCommon::StartMatchmaking);
-		//RegisterCommand(&FSessionsCommon::JoinSession);
+		RegisterCommand(&FSessionsCommon::JoinSession);
 		RegisterCommand(&FSessionsCommon::SendSessionInvite);
-		//RegisterCommand(&FSessionsCommon::RejectSessionInvite);
+		RegisterCommand(&FSessionsCommon::RejectSessionInvite);
 		RegisterCommand(&FSessionsCommon::RegisterPlayers);
 		RegisterCommand(&FSessionsCommon::UnregisterPlayers);
 	}
@@ -41,7 +41,7 @@ namespace UE::Online {
 	{
 		FGetAllSessions::Result Result;
 
-		for (const TPair<FName, TSharedRef<FSession>>& SessionPair : SessionsByName)
+		for (const TPair<FName, TSharedRef<FSession>>& SessionPair : LocalSessionsByName)
 		{
 			Result.Sessions.Add(SessionPair.Value);
 		}
@@ -51,7 +51,7 @@ namespace UE::Online {
 
 	TOnlineResult<FGetSessionByName> FSessionsCommon::GetSessionByName(FGetSessionByName::Params&& Params) const
 	{
-		if (const TSharedRef<FSession>* FoundSession = SessionsByName.Find(Params.LocalName))
+		if (const TSharedRef<FSession>* FoundSession = LocalSessionsByName.Find(Params.LocalName))
 		{
 			return TOnlineResult<FGetSessionByName>({ *FoundSession });
 		}
@@ -63,14 +63,26 @@ namespace UE::Online {
 
 	TOnlineResult<FGetSessionById> FSessionsCommon::GetSessionById(FGetSessionById::Params&& Params) const
 	{
-		if (const TSharedRef<FSession>* FoundSession = SessionsById.Find(Params.IdHandle))
+		if (const TMap<FOnlineSessionIdHandle, TSharedRef<FSession>>* UserMap = SessionSearchResultsUserMap.Find(Params.LocalUserId))
 		{
-			return TOnlineResult<FGetSessionById>({ *FoundSession });
+			if (const TSharedRef<FSession>* FoundSession = UserMap->Find(Params.IdHandle))
+			{
+				return TOnlineResult<FGetSessionById>({ *FoundSession });
+			}
 		}
-		else
+
+		if (const TMap<FOnlineSessionInviteIdHandle, TSharedRef<FSessionInvite>>* UserMap = SessionInvitesUserMap.Find(Params.LocalUserId))
 		{
-			return TOnlineResult<FGetSessionById>(Errors::NotFound());
+			for (const TPair<FOnlineSessionInviteIdHandle, TSharedRef<FSessionInvite>>& Entry : *UserMap)
+			{
+				if (Entry.Value->Session->SessionId == Params.IdHandle)
+				{
+					return TOnlineResult<FGetSessionById>({ Entry.Value->Session });
+				}
+			}
 		}
+
+		return TOnlineResult<FGetSessionById>({ Errors::NotFound() });
 	}
 
 	TOnlineAsyncOpHandle<FCreateSession> FSessionsCommon::CreateSession(FCreateSession::Params&& Params)
@@ -182,6 +194,25 @@ namespace UE::Online {
 		return Operation->GetHandle();
 	}
 
+	TOnlineResult<FGetSessionInvites> FSessionsCommon::GetSessionInvites(FGetSessionInvites::Params&& Params)
+	{
+		TArray<TSharedRef<const FSessionInvite>> SessionInvites;
+
+		if (const TMap<FOnlineSessionInviteIdHandle, TSharedRef<FSessionInvite>>* UserMap = SessionInvitesUserMap.Find(Params.LocalUserId))
+		{
+			for (const TPair<FOnlineSessionInviteIdHandle, TSharedRef<FSessionInvite>>& Entry : *UserMap)
+			{
+				SessionInvites.Add(Entry.Value);
+			}
+		}
+		else
+		{
+			return TOnlineResult<FGetSessionInvites>(Errors::NotFound());
+		}
+
+		return TOnlineResult<FGetSessionInvites>({ SessionInvites });
+	}
+
 	TOnlineAsyncOpHandle<FRejectSessionInvite> FSessionsCommon::RejectSessionInvite(FRejectSessionInvite::Params&& Params)
 	{
 		TOnlineAsyncOpRef<FRejectSessionInvite> Operation = GetOp<FRejectSessionInvite>(MoveTemp(Params));
@@ -278,7 +309,7 @@ namespace UE::Online {
 
 	TOnlineResult<FAddSessionMembers> FSessionsCommon::AddSessionMembersImpl(const FAddSessionMembers::Params& Params)
 	{
-		TSharedRef<FSession> FoundSession = SessionsByName.FindChecked(Params.SessionName);
+		TSharedRef<FSession> FoundSession = LocalSessionsByName.FindChecked(Params.SessionName);
 
 		FSessionSettings& SessionSettings = FoundSession->SessionSettings;
 
@@ -343,7 +374,7 @@ namespace UE::Online {
 
 	TOnlineResult<FRemoveSessionMembers> FSessionsCommon::RemoveSessionMembersImpl(const FRemoveSessionMembers::Params& Params)
 	{
-		TSharedRef<FSession> FoundSession = SessionsByName.FindChecked(Params.SessionName);
+		TSharedRef<FSession> FoundSession = LocalSessionsByName.FindChecked(Params.SessionName);
 
 		FSessionSettings& SessionSettings = FoundSession->SessionSettings;
 
@@ -385,7 +416,7 @@ namespace UE::Online {
 
 	TOnlineResult<FRegisterPlayers> FSessionsCommon::RegisterPlayersImpl(const FRegisterPlayers::Params& Params)
 	{
-		TSharedRef<FSession> FoundSession = SessionsByName.FindChecked(Params.SessionName);
+		TSharedRef<FSession> FoundSession = LocalSessionsByName.FindChecked(Params.SessionName);
 
 		FSessionSettings& SessionSettings = FoundSession->SessionSettings;
 
@@ -424,7 +455,7 @@ namespace UE::Online {
 
 	TOnlineResult<FUnregisterPlayers> FSessionsCommon::UnregisterPlayersImpl(const FUnregisterPlayers::Params& Params)
 	{
-		TSharedRef<FSession> FoundSession = SessionsByName.FindChecked(Params.SessionName);
+		TSharedRef<FSession> FoundSession = LocalSessionsByName.FindChecked(Params.SessionName);
 
 		FSessionSettings& SessionSettings = FoundSession->SessionSettings;
 
@@ -485,7 +516,7 @@ namespace UE::Online {
 
 	FOnlineError FSessionsCommon::CheckCreateSessionState(const FCreateSession::Params& Params)
 	{
-		if (TSharedRef<FSession>* FoundSession = SessionsByName.Find(Params.SessionName))
+		if (TSharedRef<FSession>* FoundSession = LocalSessionsByName.Find(Params.SessionName))
 		{
 			UE_LOG(LogTemp, Warning, TEXT("[FSessionsCommon::CheckCreateSessionState] Could not create session with name [%s]. A session with that name already exists"), *Params.SessionName.ToString());
 
@@ -494,7 +525,7 @@ namespace UE::Online {
 
 		if (Params.SessionSettings.bPresenceEnabled)
 		{
-			for (const TPair<FName, TSharedRef<FSession>>& Entry : SessionsByName)
+			for (const TPair<FName, TSharedRef<FSession>>& Entry : LocalSessionsByName)
 			{
 				if (Entry.Value->SessionSettings.bPresenceEnabled)
 				{
@@ -527,7 +558,7 @@ namespace UE::Online {
 
 	FOnlineError FSessionsCommon::CheckUpdateSessionState(const FUpdateSession::Params& Params)
 	{
-		if (TSharedRef<FSession>* FoundSession = SessionsByName.Find(Params.SessionName))
+		if (TSharedRef<FSession>* FoundSession = LocalSessionsByName.Find(Params.SessionName))
 		{
 			if (!FoundSession->Get().SessionSettings.bIsDedicatedServerSession)
 			{
@@ -584,24 +615,89 @@ namespace UE::Online {
 
 	FOnlineError FSessionsCommon::CheckJoinSessionParams(const FJoinSession::Params& Params)
 	{
-		const FSessionSettings& SessionSettings = Params.Session->SessionSettings;
+		if (!Params.LocalUserId.IsValid())
+		{
+			UE_LOG(LogTemp, Warning, TEXT("[FSessionsCommon::CheckJoinSessionParams] Could not join session with name [%s]. LocalUserId [%s] not valid"), *ToLogString(Params.LocalUserId));
+
+			return Errors::InvalidParams();
+		}
 
-		if (Params.Session->CurrentState != ESessionState::Valid)
+		if (!Params.LocalUsers.Contains(Params.LocalUserId))
 		{
-			UE_LOG(LogTemp, Warning, TEXT("[FSessionsEOSGS::FJoinSession] Could not join session with invalid session state [%s]"), LexToString(Params.Session->CurrentState));
+			UE_LOG(LogTemp, Warning, TEXT("[FSessionsCommon::CheckJoinSessionParams] Could not join session with name [%s]. Missing Session Member info for user [%s]"), *ToLogString(Params.LocalUserId));
 
 			return Errors::InvalidParams();
 		}
 
+		if (!Params.SessionId.IsValid())
+		{
+			UE_LOG(LogTemp, Warning, TEXT("[FSessionsCommon::CheckJoinSessionParams] Could not join session with name [%s]. SessionId [%s] not valid"), *ToLogString(Params.SessionId));
+
+			return Errors::InvalidParams();
+		}
+
+		if (Params.SessionName.ToString().IsEmpty())
+		{
+			UE_LOG(LogTemp, Warning, TEXT("[FSessionsCommon::CheckJoinSessionParams] Could not join session with no valid SessionName set"));
+
+			return Errors::InvalidParams();
+		}
+
+		return Errors::Success();
+	}
+
+	FOnlineError FSessionsCommon::CheckJoinSessionState(const FJoinSession::Params& Params)
+	{
+		// Check if a session with that name already exists
+		if (TSharedRef<FSession>* FoundSession = LocalSessionsByName.Find(Params.SessionName))
+		{
+			UE_LOG(LogTemp, Warning, TEXT("[FSessionsCommon::CheckJoinSessionState] Could not join session with name [%s]. A session with that name already exists"), *Params.SessionName.ToString());
+
+			return Errors::InvalidState(); // TODO: New error: Session with name %s already exists
+		}
+
+		// User login check for all local users
+		IAuthPtr Auth = Services.GetAuthInterface();
+
 		TArray<FOnlineAccountIdHandle> LocalUserIds;
 		LocalUserIds.Reserve(Params.LocalUsers.Num());
 		Params.LocalUsers.GenerateKeyArray(LocalUserIds);
 
+		for (const FOnlineAccountIdHandle& LocalUserId : LocalUserIds)
+		{
+			if(!Auth->IsLoggedIn(LocalUserId))
+			{
+				UE_LOG(LogTemp, Warning, TEXT("[FSessionsCommon::CheckJoinSessionState] Could not join session with user [%s] not logged in"), *ToLogString(LocalUserId));
+
+				return Errors::InvalidUser();
+			}
+		}
+
+		// We check that the session is cached and valid for a join operation by the users
+		TOnlineResult<FGetSessionById> GetSessionByIdResult = GetSessionById({ Params.LocalUserId, Params.SessionId });
+		if (GetSessionByIdResult.IsError())
+		{
+			UE_LOG(LogTemp, Warning, TEXT("[FSessionsCommon::CheckJoinSessionState] Unable to find session with SessionId parameter [%s]. Please call FindSessions to get an updated list of available sessions "), *ToLogString(Params.SessionId));
+
+			return GetSessionByIdResult.GetErrorValue();
+		}
+
+		TSharedRef<const FSession> FoundSession = GetSessionByIdResult.GetOkValue().Session;
+
+		const FSessionSettings& SessionSettings = FoundSession->SessionSettings;
+
+		if (FoundSession->CurrentState != ESessionState::Valid)
+		{
+			UE_LOG(LogTemp, Warning, TEXT("[FSessionsCommon::CheckJoinSessionState] Could not join session with invalid session state [%s]"), LexToString(FoundSession->CurrentState));
+
+			return Errors::InvalidParams();
+		}
+
 		for (const FOnlineAccountIdHandle& LocalUserId : LocalUserIds)
 		{
 			if (SessionSettings.SessionMembers.Contains(LocalUserId))
 			{
-				UE_LOG(LogTemp, Warning, TEXT("[FSessionsCommon::CheckJoinSessionParams] Could not join session. User [%s] already in session"), *ToLogString(LocalUserId));
+				UE_LOG(LogTemp, Warning, TEXT("[FSessionsCommon::CheckJoinSessionState] Could not join session. User [%s] already in session"), *ToLogString(LocalUserId));
 
 				return Errors::AccessDenied();
 			}
@@ -609,7 +705,7 @@ namespace UE::Online {
 
 		if (!SessionSettings.bAllowNewMembers)
 		{
-			UE_LOG(LogTemp, Warning, TEXT("[FSessionsCommon::CheckJoinSessionParams] Could not join session. bAllowNewMembers is set to false"));
+			UE_LOG(LogTemp, Warning, TEXT("[FSessionsCommon::CheckJoinSessionState] Could not join session. bAllowNewMembers is set to false"));
 
 			return Errors::AccessDenied();
 		}
@@ -620,7 +716,7 @@ namespace UE::Online {
 			{
 				if (!SessionSettings.RegisteredPlayers.Contains(LocalUserId))
 				{
-					UE_LOG(LogTemp, Warning, TEXT("[FSessionsCommon::CheckJoinSessionParams] Could not join session. bAllowUnregisteredPlayers is set to false and user [%s] is not registered"), *ToLogString(LocalUserId));
+					UE_LOG(LogTemp, Warning, TEXT("[FSessionsCommon::CheckJoinSessionState] Could not join session. bAllowUnregisteredPlayers is set to false and user [%s] is not registered"), *ToLogString(LocalUserId));
 
 					return Errors::AccessDenied();
 				}
@@ -629,14 +725,14 @@ namespace UE::Online {
 
 		if (SessionSettings.NumMaxPublicConnections > 0 && SessionSettings.NumOpenPublicConnections < (uint32)LocalUserIds.Num())
 		{
-			UE_LOG(LogTemp, Warning, TEXT("[FSessionsCommon::CheckJoinSessionParams] Could not join session. Not enough NumOpenPublicConnections [%u] for the amount of joining players [%d]"), SessionSettings.NumOpenPublicConnections, LocalUserIds.Num());
+			UE_LOG(LogTemp, Warning, TEXT("[FSessionsCommon::CheckJoinSessionState] Could not join session. Not enough NumOpenPublicConnections [%u] for the amount of joining players [%d]"), SessionSettings.NumOpenPublicConnections, LocalUserIds.Num());
 
 			return Errors::AccessDenied();
 		}
 
 		if (SessionSettings.NumMaxPrivateConnections > 0 && SessionSettings.NumOpenPrivateConnections < (uint32)LocalUserIds.Num())
 		{
-			UE_LOG(LogTemp, Warning, TEXT("[FSessionsCommon::CheckJoinSessionParams] Could not join session. Not enough NumOpenPrivateConnections [%u] for the amount of joining players [%d]"), SessionSettings.NumOpenPrivateConnections, LocalUserIds.Num());
+			UE_LOG(LogTemp, Warning, TEXT("[FSessionsCommon::CheckJoinSessionState] Could not join session. Not enough NumOpenPrivateConnections [%u] for the amount of joining players [%d]"), SessionSettings.NumOpenPrivateConnections, LocalUserIds.Num());
 
 			return Errors::AccessDenied();
 		}
@@ -644,36 +740,6 @@ namespace UE::Online {
 		return Errors::Success();
 	}
 
-	FOnlineError FSessionsCommon::CheckJoinSessionState(const FJoinSession::Params& Params)
-	{
-		// Check if a session with that name already exists
-		if (TSharedRef<FSession>* FoundSession = SessionsByName.Find(Params.SessionName))
-		{
-			UE_LOG(LogTemp, Warning, TEXT("[FSessionsCommon::CheckJoinSessionState] Could not join session with name [%s]. A session with that name already exists"), *Params.SessionName.ToString());
-
-			return Errors::InvalidState(); // TODO: New error: Session with name %s already exists
-		}
-
-		// User login check for all local users
-		IAuthPtr Auth = Services.GetAuthInterface();
-
-		TArray<FOnlineAccountIdHandle> LocalUserIds;
-		LocalUserIds.Reserve(Params.LocalUsers.Num());
-		Params.LocalUsers.GenerateKeyArray(LocalUserIds);
-
-		for (const FOnlineAccountIdHandle& LocalUserId : LocalUserIds)
-		{
-			if(!Auth->IsLoggedIn(LocalUserId))
-			{
-				UE_LOG(LogTemp, Warning, TEXT("[FSessionsCommon::CheckJoinSessionState] Could not join session with user [%s] not logged in"), *ToLogString(LocalUserId));
-
-				return Errors::InvalidUser();
-			}
-		}
-
-		return Errors::Success();
-	}
-
 	FOnlineError FSessionsCommon::CheckAddSessionMembersState(const FAddSessionMembers::Params& Params)
 	{
 		if (TOptional<FOnlineError> Result = CheckSessionExistsByName(Params.SessionName))
@@ -703,7 +769,7 @@ namespace UE::Online {
 	FOnlineError FSessionsCommon::CheckLeaveSessionState(const FLeaveSession::Params& Params)
 	{
 		// User login check for main caller, session check
-		if (!SessionsByName.Contains(Params.SessionName))
+		if (!LocalSessionsByName.Contains(Params.SessionName))
 		{
 			UE_LOG(LogTemp, Warning, TEXT("[FSessionsCommon::CheckLeaveSessionState] Could not leave session with name [%s]. Session not found"), *Params.SessionName.ToString());
 
@@ -729,7 +795,7 @@ namespace UE::Online {
 	FOnlineError FSessionsCommon::CheckSendSessionInviteState(const FSendSessionInvite::Params& Params)
 	{
 		// User login check for main caller, session check
-		if (TSharedRef<FSession>* FoundSession = SessionsByName.Find(Params.SessionName))
+		if (TSharedRef<FSession>* FoundSession = LocalSessionsByName.Find(Params.SessionName))
 		{
 			if (!FoundSession->Get().SessionSettings.bIsDedicatedServerSession)
 			{
@@ -798,7 +864,7 @@ namespace UE::Online {
 	{
 		TOptional<FOnlineError> Result;
 
-		if (!SessionsByName.Contains(SessionName))
+		if (!LocalSessionsByName.Contains(SessionName))
 		{
 			Result.Emplace(Errors::NotFound());
 		}
diff --git a/Engine/Plugins/Online/Experimental/OnlineServices/Source/OnlineServicesCommon/Public/Online/OnlineIdCommon.h b/Engine/Plugins/Online/Experimental/OnlineServices/Source/OnlineServicesCommon/Public/Online/OnlineIdCommon.h
--- a/Engine/Plugins/Online/Experimental/OnlineServices/Source/OnlineServicesCommon/Public/Online/OnlineIdCommon.h
+++ b/Engine/Plugins/Online/Experimental/OnlineServices/Source/OnlineServicesCommon/Public/Online/OnlineIdCommon.h
@@ -83,5 +83,8 @@ using TOnlineBasicAccountIdRegistry = TOnlineBasicIdRegistry<OnlineIdHandleTags:
 template<typename IdValueType, EOnlineServices OnlineServicesType>
 using TOnlineBasicSessionIdRegistry = TOnlineBasicIdRegistry<OnlineIdHandleTags::FSession, IdValueType, OnlineServicesType>;
 
+template<typename IdValueType, EOnlineServices OnlineServicesType>
+using TOnlineBasicSessionInviteIdRegistry = TOnlineBasicIdRegistry<OnlineIdHandleTags::FSessionInvite, IdValueType, OnlineServicesType>;
+
 /* UE::Online */ }
 
diff --git a/Engine/Plugins/Online/Experimental/OnlineServices/Source/OnlineServicesCommon/Public/Online/SessionsCommon.h b/Engine/Plugins/Online/Experimental/OnlineServices/Source/OnlineServicesCommon/Public/Online/SessionsCommon.h
--- a/Engine/Plugins/Online/Experimental/OnlineServices/Source/OnlineServicesCommon/Public/Online/SessionsCommon.h
+++ b/Engine/Plugins/Online/Experimental/OnlineServices/Source/OnlineServicesCommon/Public/Online/SessionsCommon.h
@@ -62,6 +62,56 @@ public:
 	TOnlineBasicSessionIdRegistry<FString, OnlineServicesType> BasicRegistry;
 };
 
+template<EOnlineServices OnlineServicesType>
+class TOnlineSessionInviteIdStringRegistry : public IOnlineSessionInviteIdRegistry
+{
+public:
+	// Begin IOnlineSessionIdRegistry
+	virtual inline FString ToLogString(const FOnlineSessionInviteIdHandle& Handle) const override
+	{
+		FString IdValue = BasicRegistry.FindIdValue(Handle);
+
+		if (IdValue.Len() == 0)
+		{
+			IdValue = FString(TEXT("[InvalidSessionID]"));
+		}
+
+		return IdValue;
+	};
+
+	virtual inline TArray<uint8> ToReplicationData(const FOnlineSessionInviteIdHandle& Handle) const override
+	{
+		const FString IdValue = BasicRegistry.FindIdValue(Handle);
+		const FTCHARToUTF8 IdValueUtf8(IdValue);
+
+		TArray<uint8> ReplicationData;
+		ReplicationData.SetNumUninitialized(IdValueUtf8.Length());
+
+		FMemory::Memcpy(ReplicationData.GetData(), IdValueUtf8.Get(), IdValueUtf8.Length());
+
+		return ReplicationData;
+	}
+
+	virtual inline FOnlineSessionInviteIdHandle FromReplicationData(const TArray<uint8>& ReplicationData) override
+	{
+		const FUTF8ToTCHAR IdValueTCHAR((char*)ReplicationData.GetData(), ReplicationData.Num());
+		const FString IdValue = FString(IdValueTCHAR.Length(), IdValueTCHAR.Get());
+
+		if (!IdValue.IsEmpty())
+		{
+			return BasicRegistry.FindOrAddHandle(IdValue);
+		}
+
+		return FOnlineSessionInviteIdHandle();
+	}
+	// End IOnlineSessionIdRegistry
+
+	virtual ~TOnlineSessionInviteIdStringRegistry() = default;
+
+public:
+	TOnlineBasicSessionInviteIdRegistry<FString, OnlineServicesType> BasicRegistry;
+};
+
 class ONLINESERVICESCOMMON_API FSessionsCommon : public TOnlineComponent<ISessions>
 {
 public:
@@ -86,6 +136,7 @@ public:
 	virtual TOnlineAsyncOpHandle<FAddSessionMembers> AddSessionMembers(FAddSessionMembers::Params&& Params) override;
 	virtual TOnlineAsyncOpHandle<FRemoveSessionMembers> RemoveSessionMembers(FRemoveSessionMembers::Params&& Params) override;
 	virtual TOnlineAsyncOpHandle<FSendSessionInvite> SendSessionInvite(FSendSessionInvite::Params&& Params) override;
+	virtual TOnlineResult<FGetSessionInvites> GetSessionInvites(FGetSessionInvites::Params&& Params) override;
 	virtual TOnlineAsyncOpHandle<FRejectSessionInvite> RejectSessionInvite(FRejectSessionInvite::Params&& Params) override;
 	virtual TOnlineAsyncOpHandle<FRegisterPlayers> RegisterPlayers(FRegisterPlayers::Params&& Params) override;
 	virtual TOnlineAsyncOpHandle<FUnregisterPlayers> UnregisterPlayers(FUnregisterPlayers::Params&& Params) override;
@@ -142,8 +193,12 @@ protected:
 		TOnlineEventCallable<void(const FUISessionJoinRequested&)> OnUISessionJoinRequested;
 	} SessionEvents;
 
-	TMap<FName, TSharedRef<FSession>> SessionsByName;
-	TMap<FOnlineSessionIdHandle, TSharedRef<FSession>> SessionsById;
+	/** Map of sessions that local users are part of, indexed by their local name */
+	TMap<FName, TSharedRef<FSession>> LocalSessionsByName;
+	/** Cache for the last set of session search results, mapped per user */
+	TMap<FOnlineAccountIdHandle, TMap<FOnlineSessionIdHandle, TSharedRef<FSession>>> SessionSearchResultsUserMap;
+	/** Cache for received session invites, mapped per user */
+	TMap<FOnlineAccountIdHandle, TMap<FOnlineSessionInviteIdHandle, TSharedRef<FSessionInvite>>> SessionInvitesUserMap;
 
 	TSharedPtr<FFindSessions::Result> CurrentSessionSearch;
 	TSharedPtr<TOnlineAsyncOp<FFindSessions>> CurrentSessionSearchHandle;
diff --git a/Engine/Plugins/Online/Experimental/OnlineServices/Source/OnlineServicesCommonEngineUtils/Private/Online/SessionsLAN.cpp b/Engine/Plugins/Online/Experimental/OnlineServices/Source/OnlineServicesCommonEngineUtils/Private/Online/SessionsLAN.cpp
--- a/Engine/Plugins/Online/Experimental/OnlineServices/Source/OnlineServicesCommonEngineUtils/Private/Online/SessionsLAN.cpp
+++ b/Engine/Plugins/Online/Experimental/OnlineServices/Source/OnlineServicesCommonEngineUtils/Private/Online/SessionsLAN.cpp
@@ -23,6 +23,11 @@ FOnlineSessionIdHandle FOnlineSessionIdRegistryLAN::GetNextSessionId()
 	return BasicRegistry.FindOrAddHandle(FGuid::NewGuid().ToString());
 }
 
+bool FOnlineSessionIdRegistryLAN::IsSessionIdExpired(const FOnlineSessionIdHandle& InHandle) const
+{
+	return BasicRegistry.FindIdValue(InHandle).IsEmpty();
+}
+
 /** FSessionLAN */
 
 FSessionLAN::FSessionLAN()
@@ -129,8 +134,7 @@ TOnlineAsyncOpHandle<FCreateSession> FSessionsLAN::CreateSession(FCreateSession:
 		// For LAN sessions, we'll add all the Session members manually instead of calling JoinSession since there is no API calls involved
 		NewSessionLANRef->SessionSettings.SessionMembers.Append(OpParams.LocalUsers);
 
-		SessionsByName.Add(OpParams.SessionName, NewSessionLANRef);
-		SessionsById.Add(NewSessionLANRef->SessionId, NewSessionLANRef);
+		LocalSessionsByName.Add(OpParams.SessionName, NewSessionLANRef);
 
 		Op.SetResult(FCreateSession::Result{ NewSessionLANRef });
 	})
@@ -156,7 +160,7 @@ TOnlineAsyncOpHandle<FUpdateSession> FSessionsLAN::UpdateSession(FUpdateSession:
 
 		// We'll check and update all settings one by one, with additional logic wherever required
 
-		TSharedRef<FSession>& FoundSession = SessionsByName.FindChecked(OpParams.SessionName);
+		TSharedRef<FSession>& FoundSession = LocalSessionsByName.FindChecked(OpParams.SessionName);
 		FSessionSettings& SessionSettings = FoundSession->SessionSettings;
 		const FSessionSettingsUpdate& UpdatedSettings = OpParams.Mutations;
 
@@ -248,11 +252,23 @@ TOnlineAsyncOpHandle<FJoinSession> FSessionsLAN::JoinSession(FJoinSession::Param
 			return;
 		}
 
-		// TODO: This should retrieve an FSession from a cache and move it, instead of making a copy
-		TSharedRef<FSession> NewSessionRef = MakeShared<FSession>(*OpParams.Session);
+		TOnlineResult<FGetSessionById> GetSessionByIdResult = GetSessionById({ OpParams.LocalUserId, OpParams.SessionId });
+		if (GetSessionByIdResult.IsError())
+		{
+			// If no result is found, the id might be expired, which we should notify
+			if (FOnlineSessionIdRegistryLAN::Get().IsSessionIdExpired(OpParams.SessionId))
+			{
+				UE_LOG(LogTemp, Warning, TEXT("[FSessionsLAN::JoinSession] SessionId parameter [%s] is expired. Please call FindSessions to get an updated list of available sessions "), *ToLogString(OpParams.SessionId));
+			}
 
-		SessionsByName.Add(OpParams.SessionName, NewSessionRef);
-		SessionsById.Add(NewSessionRef->SessionId, NewSessionRef);
+			Op.SetError(MoveTemp(GetSessionByIdResult.GetErrorValue()));
+			return ;
+		}
+
+		const TSharedRef<const FSession>& FoundSession = GetSessionByIdResult.GetOkValue().Session;
+		TSharedRef<FSession> NewSessionRef = ConstCastSharedRef<FSession>(FoundSession);
+
+		LocalSessionsByName.Add(OpParams.SessionName, NewSessionRef);
 
 		Op.SetResult(FJoinSession::Result{ NewSessionRef });
 
@@ -286,7 +302,7 @@ TOnlineAsyncOpHandle<FLeaveSession> FSessionsLAN::LeaveSession(FLeaveSession::Pa
 
 		// This scope guarantees that FoundSession can't be used after its removal
 		{
-			TSharedRef<FSession>& FoundSession = SessionsByName.FindChecked(OpParams.SessionName);
+			TSharedRef<FSession>& FoundSession = LocalSessionsByName.FindChecked(OpParams.SessionName);
 
 			// Only the host should stop the session at this point
 			if (FoundSession->OwnerUserId == OpParams.LocalUserId && FoundSession->SessionSettings.JoinPolicy == ESessionJoinPolicy::Public)
@@ -294,10 +310,7 @@ TOnlineAsyncOpHandle<FLeaveSession> FSessionsLAN::LeaveSession(FLeaveSession::Pa
 				StopLANSession();
 			}
 
-			// TODO: Remove all local users from session
-
-			SessionsById.Remove(FoundSession->SessionId);
-			SessionsByName.Remove(OpParams.SessionName);
+			LocalSessionsByName.Remove(OpParams.SessionName);
 		}
 
 		Op.SetResult(FLeaveSession::Result{ });
@@ -400,7 +413,7 @@ void FSessionsLAN::StopLANSession()
 void FSessionsLAN::OnValidQueryPacketReceived(uint8* PacketData, int32 PacketLength, uint64 ClientNonce)
 {
 	// Iterate through all registered sessions and respond for each one that can be joinable
-	for (const TPair<FName, TSharedRef<FSession>>& Entry : SessionsByName)
+	for (const TPair<FName, TSharedRef<FSession>>& Entry : LocalSessionsByName)
 	{
 		const TSharedRef<FSession>& Session = Entry.Value;
 
@@ -436,6 +449,9 @@ void FSessionsLAN::OnValidResponsePacketReceived(uint8* PacketData, int32 Packet
 		ReadSessionFromPacket(Packet, *Session);
 
 		CurrentSessionSearch->FoundSessions.Add(Session);
+
+		TMap<FOnlineSessionIdHandle, TSharedRef<FSession>>& UserMap = SessionSearchResultsUserMap.FindOrAdd(CurrentSessionSearchHandle->GetParams().LocalUserId);
+		UserMap.Emplace(Session->SessionId, Session);
 	}
 }
 
diff --git a/Engine/Plugins/Online/Experimental/OnlineServices/Source/OnlineServicesCommonEngineUtils/Public/Online/SessionsLAN.h b/Engine/Plugins/Online/Experimental/OnlineServices/Source/OnlineServicesCommonEngineUtils/Public/Online/SessionsLAN.h
--- a/Engine/Plugins/Online/Experimental/OnlineServices/Source/OnlineServicesCommonEngineUtils/Public/Online/SessionsLAN.h
+++ b/Engine/Plugins/Online/Experimental/OnlineServices/Source/OnlineServicesCommonEngineUtils/Public/Online/SessionsLAN.h
@@ -16,6 +16,8 @@ public:
 	static FOnlineSessionIdRegistryLAN& Get();
 
 	FOnlineSessionIdHandle GetNextSessionId();
+
+	bool IsSessionIdExpired(const FOnlineSessionIdHandle& InHandle) const;
 };
 
 class ONLINESERVICESCOMMONENGINEUTILS_API FSessionLAN : public FSession
diff --git a/Engine/Plugins/Online/Experimental/OnlineServices/Source/OnlineServicesInterface/Public/Online/Sessions.h b/Engine/Plugins/Online/Experimental/OnlineServices/Source/OnlineServicesInterface/Public/Online/Sessions.h
--- a/Engine/Plugins/Online/Experimental/OnlineServices/Source/OnlineServicesInterface/Public/Online/Sessions.h
+++ b/Engine/Plugins/Online/Experimental/OnlineServices/Source/OnlineServicesInterface/Public/Online/Sessions.h
@@ -250,8 +250,8 @@ struct FSessionInvite
 	/* The user which sent the invite */
 	FOnlineAccountIdHandle SenderId;
 
-	/* The invite id, needed for retrieving session information and rejecting the invite */
-	FString InviteId;
+	/* The invite id handle, needed for retrieving session information and rejecting the invite */
+	FOnlineSessionInviteIdHandle InviteId;
 
 	/* Pointer to the session information */
 	TSharedRef<const FSession> Session;
@@ -297,6 +297,8 @@ struct FGetSessionById
 
 	struct Params
 	{
+		FOnlineAccountIdHandle LocalUserId;
+
 		FOnlineSessionIdHandle IdHandle;
 	};
 
@@ -457,8 +459,8 @@ struct FJoinSession
 		/* Local name for the session */
 		FName SessionName;
 
-		/* A reference to the session to be joined. */
-		TSharedRef<const FSession> Session; // TODO: FOnlineSessionIdHandle to be used after we cache search results
+		/* Id handle for the session to be joined. To be retrieved via session search or invite */
+		FOnlineSessionIdHandle SessionId;
 
 		/* Information for all local users who will join the session (includes the session creator). Any player that wants to join the session needs defined information to become a new member */
 		FSessionMembersMap LocalUsers;
@@ -544,6 +546,23 @@ struct FSendSessionInvite
 	};
 };
 
+struct FGetSessionInvites
+{
+	static constexpr TCHAR Name[] = TEXT("GetSessionInvites");
+
+	struct Params
+	{
+		/* The local user agent */
+		FOnlineAccountIdHandle LocalUserId;
+	};
+
+	struct Result
+	{
+		/** Set of active session invites */
+		TArray<TSharedRef<const FSessionInvite>> SessionInvites;
+	};
+};
+
 struct FRejectSessionInvite
 {
 	static constexpr TCHAR Name[] = TEXT("RejectSessionInvite");
@@ -553,8 +572,8 @@ struct FRejectSessionInvite
 		/* The local user agent which started the query*/
 		FOnlineAccountIdHandle LocalUserId;
 
-		/* The invite to be rejected */
-		TSharedRef<const FSessionInvite> SessionInvite; // TODO: FOnlineSessionInviteIdHandle to be used after we cache invites
+		/* The id handle for the invite to be rejected */
+		FOnlineSessionInviteIdHandle SessionInviteId;
 	};
 
 	struct Result
@@ -775,6 +794,14 @@ public:
 	 */
 	virtual TOnlineAsyncOpHandle<FSendSessionInvite> SendSessionInvite(FSendSessionInvite::Params&& Params) = 0;
 
+	/**
+	 * Returns all cached session invites for the given user.
+	 *
+	 * @param Parameters for the SendSessionInvite call
+	 * @return
+	 */
+	virtual TOnlineResult<FGetSessionInvites> GetSessionInvites(FGetSessionInvites::Params&& Params) = 0;
+
 	/**
 	 * Rejects a given session invite for a user.
 	 *
@@ -937,6 +964,7 @@ BEGIN_ONLINE_STRUCT_META(FGetSessionByName::Result)
 END_ONLINE_STRUCT_META()
 
 BEGIN_ONLINE_STRUCT_META(FGetSessionById::Params)
+	ONLINE_STRUCT_FIELD(FGetSessionById::Params, LocalUserId),
 	ONLINE_STRUCT_FIELD(FGetSessionById::Params, IdHandle)
 END_ONLINE_STRUCT_META()
 
@@ -1003,7 +1031,7 @@ END_ONLINE_STRUCT_META()
 BEGIN_ONLINE_STRUCT_META(FJoinSession::Params)
 	ONLINE_STRUCT_FIELD(FJoinSession::Params, LocalUserId),
 	ONLINE_STRUCT_FIELD(FJoinSession::Params, SessionName),
-	ONLINE_STRUCT_FIELD(FJoinSession::Params, Session),
+	ONLINE_STRUCT_FIELD(FJoinSession::Params, SessionId),
 	ONLINE_STRUCT_FIELD(FJoinSession::Params, LocalUsers)
 END_ONLINE_STRUCT_META()
 
@@ -1040,9 +1068,17 @@ END_ONLINE_STRUCT_META()
 BEGIN_ONLINE_STRUCT_META(FSendSessionInvite::Result)
 END_ONLINE_STRUCT_META()
 
+BEGIN_ONLINE_STRUCT_META(FGetSessionInvites::Params)
+	ONLINE_STRUCT_FIELD(FGetSessionInvites::Params, LocalUserId)
+END_ONLINE_STRUCT_META()
+
+BEGIN_ONLINE_STRUCT_META(FGetSessionInvites::Result)
+	ONLINE_STRUCT_FIELD(FGetSessionInvites::Result, SessionInvites)
+END_ONLINE_STRUCT_META()
+
 BEGIN_ONLINE_STRUCT_META(FRejectSessionInvite::Params)
 	ONLINE_STRUCT_FIELD(FRejectSessionInvite::Params, LocalUserId),
-	ONLINE_STRUCT_FIELD(FRejectSessionInvite::Params, SessionInvite)
+	ONLINE_STRUCT_FIELD(FRejectSessionInvite::Params, SessionInviteId)
 END_ONLINE_STRUCT_META()
 
 BEGIN_ONLINE_STRUCT_META(FRejectSessionInvite::Result)
diff --git a/Engine/Plugins/Online/Experimental/OnlineServicesEOSGS/Source/Private/Online/OnlineServicesEOSGS.cpp b/Engine/Plugins/Online/Experimental/OnlineServicesEOSGS/Source/Private/Online/OnlineServicesEOSGS.cpp
--- a/Engine/Plugins/Online/Experimental/OnlineServicesEOSGS/Source/Private/Online/OnlineServicesEOSGS.cpp
+++ b/Engine/Plugins/Online/Experimental/OnlineServicesEOSGS/Source/Private/Online/OnlineServicesEOSGS.cpp
@@ -130,7 +130,7 @@ TOnlineResult<FGetResolvedConnectString> FOnlineServicesEOSGS::GetResolvedConnec
 		ISessionsPtr SessionsEOS = GetSessionsInterface();
 		check(SessionsEOS);
 
-		TOnlineResult<FGetSessionById> Result = SessionsEOS->GetSessionById({ Params.SessionId });
+		TOnlineResult<FGetSessionById> Result = SessionsEOS->GetSessionById({ Params.LocalUserId, Params.SessionId });
 		if (Result.IsOk())
 		{
 #if WITH_ENGINE
diff --git a/Engine/Plugins/Online/Experimental/OnlineServicesEOSGS/Source/Private/Online/SessionsEOSGS.cpp b/Engine/Plugins/Online/Experimental/OnlineServicesEOSGS/Source/Private/Online/SessionsEOSGS.cpp
--- a/Engine/Plugins/Online/Experimental/OnlineServicesEOSGS/Source/Private/Online/SessionsEOSGS.cpp
+++ b/Engine/Plugins/Online/Experimental/OnlineServicesEOSGS/Source/Private/Online/SessionsEOSGS.cpp
@@ -36,6 +36,24 @@ FOnlineSessionIdRegistryEOSGS& FOnlineSessionIdRegistryEOSGS::Get()
 	return Instance;
 }
 
+bool FOnlineSessionIdRegistryEOSGS::IsSessionIdExpired(const FOnlineSessionIdHandle& InHandle) const
+{
+	return BasicRegistry.FindIdValue(InHandle).IsEmpty();
+}
+
+/** FOnlineSessionInviteIdRegistryEOSGS */
+
+FOnlineSessionInviteIdRegistryEOSGS& FOnlineSessionInviteIdRegistryEOSGS::Get()
+{
+	static FOnlineSessionInviteIdRegistryEOSGS Instance;
+	return Instance;
+}
+
+bool FOnlineSessionInviteIdRegistryEOSGS::IsSessionInviteIdExpired(const FOnlineSessionInviteIdHandle& InHandle) const
+{
+	return BasicRegistry.FindIdValue(InHandle).IsEmpty();
+}
+
 /** FSessionEOSGS */
 
 FSessionEOSGS::FSessionEOSGS(const EOS_HSessionDetails& InSessionDetailsHandle)
@@ -142,7 +160,7 @@ void FSessionsEOSGS::HandleSessionInviteReceived(const EOS_Sessions_SessionInvit
 	if (LocalUserId.IsValid())
 	{
 		Services.Get<FAuthEOSGS>()->ResolveAccountIds(LocalUserId, { Data->LocalUserId, Data->TargetUserId })
-			.Next([this, WeakThis = AsWeak(), LocalUserId, InviteId = FString(Data->InviteId)](TArray<FOnlineAccountIdHandle> ResolvedAccountIds)
+			.Next([this, WeakThis = AsWeak(), LocalUserId, InviteId = FString(UTF8_TO_TCHAR(Data->InviteId))](TArray<FOnlineAccountIdHandle> ResolvedAccountIds)
 		{
 			if (const TSharedPtr<ISessions> StrongThis = WeakThis.Pin())
 			{
@@ -150,11 +168,23 @@ void FSessionsEOSGS::HandleSessionInviteReceived(const EOS_Sessions_SessionInvit
 				const FOnlineAccountIdHandle& ReceiverId = ResolvedAccountIds[0];
 				const FOnlineAccountIdHandle& SenderId = ResolvedAccountIds[1];
 
-				TResult<TSharedRef<const FSessionInvite>, FOnlineError> SessionInviteResult = BuildSessionInvite(ReceiverId, SenderId, InviteId);
-
+				TResult<TSharedRef<const FSession>, FOnlineError> SessionInviteResult = BuildSessionFromInvite(InviteId);
+				
 				if (SessionInviteResult.IsOk())
 				{
-					FSessionInviteReceived Event = { LocalUserId, SessionInviteResult.GetOkValue() };
+					FOnlineSessionInviteIdHandle InviteIdHandle = CreateSessionInviteId(InviteId);
+
+					TSharedRef<FSessionInvite> SessionInviteRef = MakeShared<FSessionInvite>(FSessionInvite{
+						ReceiverId,
+						SenderId,
+						InviteIdHandle,
+						SessionInviteResult.GetOkValue()
+					});
+
+					TMap<FOnlineSessionInviteIdHandle, TSharedRef<FSessionInvite>>& UserMap = SessionInvitesUserMap.FindOrAdd(ReceiverId);
+					UserMap.Emplace(InviteIdHandle, SessionInviteRef);
+
+					FSessionInviteReceived Event = { LocalUserId, SessionInviteRef };
 
 					SessionEvents.OnSessionInviteReceived.Broadcast(Event);
 				}
@@ -298,8 +328,7 @@ TOnlineAsyncOpHandle<FCreateSession> FSessionsEOSGS::CreateSession(FCreateSessio
 		NewSessionEOSGSRef->OwnerUserId = OpParams.LocalUserId;
 		NewSessionEOSGSRef->SessionSettings = OpParams.SessionSettings;
 
- 		SessionsByName.Add(OpParams.SessionName, NewSessionEOSGSRef);
- 		SessionsById.Add(NewSessionEOSGSRef->SessionId, NewSessionEOSGSRef);
+		LocalSessionsByName.Add(OpParams.SessionName, NewSessionEOSGSRef);
 
 		FUpdateSessionImpl::Params UpdateSessionImplParams { MakeShared<FSessionModificationHandleEOSGS>(SessionModificationHandle) };
 		
@@ -311,7 +340,7 @@ TOnlineAsyncOpHandle<FCreateSession> FSessionsEOSGS::CreateSession(FCreateSessio
 
 		if (Result.IsOk())
 		{
-			TSharedRef<FSession>& FoundSession = SessionsByName.FindChecked(OpParams.SessionName);
+			TSharedRef<FSession>& FoundSession = LocalSessionsByName.FindChecked(OpParams.SessionName);
 			FoundSession->CurrentState = ESessionState::Valid;
 
 			// TODO: Make additional local users join the created session
@@ -322,7 +351,7 @@ TOnlineAsyncOpHandle<FCreateSession> FSessionsEOSGS::CreateSession(FCreateSessio
 		else
 		{	
 			// If creation failed, we remove the session objects
-			TryRemoveSession(OpParams.SessionName);
+			LocalSessionsByName.Remove(OpParams.SessionName);
 
 			Op.SetError(MoveTemp(Result.GetErrorValue()));
 		}
@@ -441,7 +470,7 @@ void FSessionsEOSGS::WriteUpdateSessionModificationHandle(EOS_HSessionModificati
 // 	NewSetting.Visibility = ECustomSessionSettingVisibility::ViaOnlineService;
 // 
 // 	AddAttribute(SessionModificationHandle, EOS_SESSIONS_ALLOW_NEW_MEMBERS, NewSetting);
-	const TSharedRef<FSession>& FoundSession = SessionsByName.FindChecked(SessionName);
+	const TSharedRef<FSession>& FoundSession = LocalSessionsByName.FindChecked(SessionName);
 
 	if (NewSettings.JoinPolicy.IsSet())
 	{
@@ -598,7 +627,7 @@ TOnlineAsyncOpHandle<FUpdateSession> FSessionsEOSGS::UpdateSession(FUpdateSessio
 		{
 			const FUpdateSession::Params& OpParams = Op.GetParams();
 
-			if (TSharedRef<FSession>* Session = SessionsByName.Find(OpParams.SessionName))
+			if (TSharedRef<FSession>* Session = LocalSessionsByName.Find(OpParams.SessionName))
 			{
 				// Now that the API Session update has processed successfully, we'll update our local session with the same data
 				Session->Get().SessionSettings += OpParams.Mutations;
@@ -644,7 +673,7 @@ TFuture<TDefaultErrorResult<FUpdateSessionImpl>> FSessionsEOSGS::UpdateSessionIm
 			return;
 		}
 
-		if (TSharedRef<FSession>* FoundSession = SessionsByName.Find(FName(Result->SessionName)))
+		if (TSharedRef<FSession>* FoundSession = LocalSessionsByName.Find(FName(Result->SessionName)))
 		{
 			// In session creation calls, we'll need to set the new id for the session
 			if (!FoundSession->Get().SessionId.IsValid())
@@ -703,19 +732,15 @@ TOnlineAsyncOpHandle<FLeaveSession> FSessionsEOSGS::LeaveSession(FLeaveSession::
 
 		const FLeaveSession::Params& OpParams = Op.GetParams();
 
-		if (TryRemoveSession(OpParams.SessionName))
-		{
-			Op.SetResult(FLeaveSession::Result{ });
+		LocalSessionsByName.Remove(OpParams.SessionName);
+
+		Op.SetResult(FLeaveSession::Result{ });
+
+		FSessionLeft SessionLeftEvent;
+		SessionLeftEvent.LocalUserIds.Add(OpParams.LocalUserId);
+		//SessionLeftEvent.LocalUserIds.Append(OpParams.AdditionalLocalUsers); // TODO: Process multiple users in Session destruction
+		SessionEvents.OnSessionLeft.Broadcast(SessionLeftEvent);
 
-			FSessionLeft SessionLeftEvent;
-			SessionLeftEvent.LocalUserIds.Add(OpParams.LocalUserId);
-			//SessionLeftEvent.LocalUserIds.Append(OpParams.AdditionalLocalUsers); // TODO: Process multiple users in Session destruction
-			SessionEvents.OnSessionLeft.Broadcast(SessionLeftEvent);
-		}
-		else
-		{
-			Op.SetError(Errors::NotFound());
-		}
 	})
 	.Enqueue(GetSerialQueue());
 
@@ -799,6 +824,11 @@ FOnlineSessionIdHandle FSessionsEOSGS::CreateSessionId(const FString& SessionId)
 	return FOnlineSessionIdRegistryEOSGS::Get().BasicRegistry.FindOrAddHandle(SessionId);
 }
 
+FOnlineSessionInviteIdHandle FSessionsEOSGS::CreateSessionInviteId(const FString& SessionInviteId) const
+{
+	return FOnlineSessionInviteIdRegistryEOSGS::Get().BasicRegistry.FindOrAddHandle(SessionInviteId);
+}
+
 TOnlineAsyncOpHandle<FFindSessions> FSessionsEOSGS::FindSessions(FFindSessions::Params&& Params)
 {
 	// LAN Sessions
@@ -832,9 +862,13 @@ TOnlineAsyncOpHandle<FFindSessions> FSessionsEOSGS::FindSessions(FFindSessions::
 			return;
 		}
 
+		// Before we start the search, we reset the variables and cache
 		CurrentSessionSearch = MakeShared<FFindSessions::Result>();
 		CurrentSessionSearchHandle = Op.AsShared();
 
+		SessionSearchResultsUserMap.FindOrAdd(OpParams.LocalUserId).Reset();
+
+		// We start preparing the search
 		EOS_Sessions_CreateSessionSearchOptions CreateSessionSearchOptions = { };
 		CreateSessionSearchOptions.ApiVersion = EOS_SESSIONS_CREATESESSIONSEARCH_API_LATEST;
 		static_assert(EOS_SESSIONS_CREATESESSIONSEARCH_API_LATEST == 1, "EOS_Sessions_CreateSessionSearchOptions updated, check new fields");
@@ -893,6 +927,9 @@ TOnlineAsyncOpHandle<FFindSessions> FSessionsEOSGS::FindSessions(FFindSessions::
 				const TSharedRef<FSessionEOSGS>& NewSessionEOSGSRef = MakeShared<FSessionEOSGS>(SessionDetailsHandle);
 
 				CurrentSessionSearch->FoundSessions.Add(NewSessionEOSGSRef);
+
+				TMap<FOnlineSessionIdHandle, TSharedRef<FSession>>& UserMap = SessionSearchResultsUserMap.FindOrAdd(Op.GetParams().LocalUserId);
+				UserMap.Emplace(NewSessionEOSGSRef->SessionId, NewSessionEOSGSRef);
 			}
 			else
 			{
@@ -912,22 +949,31 @@ TOnlineAsyncOpHandle<FFindSessions> FSessionsEOSGS::FindSessions(FFindSessions::
 
 TOnlineAsyncOpHandle<FJoinSession> FSessionsEOSGS::JoinSession(FJoinSession::Params&& Params)
 {
+	TOnlineAsyncOpRef<FJoinSession> Op = GetOp<FJoinSession>(MoveTemp(Params));
+	const FJoinSession::Params& OpParams = Op->GetParams();
+
+	TOnlineResult<FGetSessionById> GetSessionByIdResult = GetSessionById({ OpParams.LocalUserId, OpParams.SessionId });
+	if (GetSessionByIdResult.IsError())
+	{
+		Op->SetError(MoveTemp(GetSessionByIdResult.GetErrorValue()));
+		return Op->GetHandle();
+	}
+
+	const TSharedRef<const FSession>& FoundSession = GetSessionByIdResult.GetOkValue().Session;
+
 	// LAN Sessions
-	if (Params.Session->SessionSettings.bIsLANSession)
+	if (FoundSession->SessionSettings.bIsLANSession)
 	{
 		return FSessionsLAN::JoinSession(MoveTemp(Params));
 	}
 
 	// EOSGS Sessions
 
-	TOnlineAsyncOpRef<FJoinSession> Op = GetOp<FJoinSession>(MoveTemp(Params));
-	const FJoinSession::Params& OpParams = Op->GetParams();
-
 	// We check that the passed session has a valid details handle
-	const FSessionEOSGS& SessionEOSGS = FSessionEOSGS::Cast(*Params.Session);
+	const FSessionEOSGS& SessionEOSGS = FSessionEOSGS::Cast(*FoundSession);
 	if (!SessionEOSGS.SessionDetailsHandle.IsValid())
 	{
-		UE_LOG(LogTemp, Warning, TEXT("[FSessionsEOSGS::FJoinSession] Could not join session with invalid session details handle in session with id [%s]"), *ToLogString(Params.Session->SessionId));
+		UE_LOG(LogTemp, Warning, TEXT("[FSessionsEOSGS::FJoinSession] Could not join session with invalid session details handle in session with id [%s]"), *ToLogString(OpParams.SessionId));
 
 		Op->SetError(Errors::InvalidParams());
 		return Op->GetHandle();
@@ -953,26 +999,33 @@ TOnlineAsyncOpHandle<FJoinSession> FSessionsEOSGS::JoinSession(FJoinSession::Par
 			return;
 		}
 
-		// We make a copy of the Session object		
-		TSharedRef<FSession> NewSessionRef = MakeShared<FSession>(*OpParams.Session);
+		TOnlineResult<FGetSessionById> GetSessionByIdResult = GetSessionById({ OpParams.LocalUserId, OpParams.SessionId });
+		if (GetSessionByIdResult.IsError())
+		{
+			Op.SetError(MoveTemp(GetSessionByIdResult.GetErrorValue()));
+			Promise.EmplaceValue();
+			return;
+		}
+
+		const TSharedRef<const FSession>& FoundSession = GetSessionByIdResult.GetOkValue().Session;
+		TSharedRef<FSession> NewSessionRef = ConstCastSharedRef<FSession>(FoundSession);
 		NewSessionRef->CurrentState = ESessionState::Joining;
 
-		SessionsByName.Add(OpParams.SessionName, NewSessionRef);
-		SessionsById.Add(NewSessionRef->SessionId, NewSessionRef);
+		LocalSessionsByName.Add(OpParams.SessionName, NewSessionRef);
 
 		// We start setup for the API call
 		EOS_Sessions_JoinSessionOptions JoinSessionOptions = { };
 		JoinSessionOptions.ApiVersion = EOS_SESSIONS_JOINSESSION_API_LATEST;
 		static_assert(EOS_SESSIONS_JOINSESSION_API_LATEST == 2, "EOS_Sessions_JoinSessionOptions updated, check new fields");
 
-		JoinSessionOptions.bPresenceEnabled = OpParams.Session->SessionSettings.bPresenceEnabled;
+		JoinSessionOptions.bPresenceEnabled = FoundSession->SessionSettings.bPresenceEnabled;
 
 		JoinSessionOptions.LocalUserId = GetProductUserIdChecked(OpParams.LocalUserId);
 
 		FTCHARToUTF8 SessionNameUtf8(OpParams.SessionName.ToString());
 		JoinSessionOptions.SessionName = SessionNameUtf8.Get();
 
-		const FSessionEOSGS& SessionEOSGS = FSessionEOSGS::Cast(*OpParams.Session);
+		const FSessionEOSGS& SessionEOSGS = FSessionEOSGS::Cast(*FoundSession);
 		JoinSessionOptions.SessionHandle = SessionEOSGS.SessionDetailsHandle->SessionDetailsHandle;
 
 		EOS_Async(EOS_Sessions_JoinSession, SessionsHandle, JoinSessionOptions, MoveTemp(Promise));
@@ -983,16 +1036,33 @@ TOnlineAsyncOpHandle<FJoinSession> FSessionsEOSGS::JoinSession(FJoinSession::Par
 
 		if (Result->ResultCode != EOS_EResult::EOS_Success)
 		{
-			TryRemoveSession(OpParams.SessionName);
+			LocalSessionsByName.Remove(OpParams.SessionName);
 
 			Op.SetError(Errors::FromEOSResult(Result->ResultCode));
 			return;
 		}
 		
-		if (TSharedRef<FSession>* FoundSession = SessionsByName.Find(OpParams.SessionName))
+		if (TSharedRef<FSession>* FoundSession = LocalSessionsByName.Find(OpParams.SessionName))
 		{
 			FoundSession->Get().CurrentState = ESessionState::Valid;
 
+			// After successfully joining a session, we'll remove all related invites if any are found
+			if (TMap<FOnlineSessionInviteIdHandle, TSharedRef<FSessionInvite>>* UserMap = SessionInvitesUserMap.Find(OpParams.LocalUserId))
+			{
+				TArray<FOnlineSessionInviteIdHandle> InviteIdsToRemove;
+				for (const TPair<FOnlineSessionInviteIdHandle, TSharedRef<FSessionInvite>>& Entry : *UserMap)
+				{
+					if (Entry.Value->Session->SessionId == FoundSession->Get().SessionId)
+					{
+						InviteIdsToRemove.Add(Entry.Key);
+					}
+				}
+				for (const FOnlineSessionInviteIdHandle& InviteId : InviteIdsToRemove)
+				{
+					UserMap->Remove(InviteId);
+				}
+			}
+
 			FSessionJoined Event = { { OpParams.LocalUserId }, *FoundSession };
 
 			SessionEvents.OnSessionJoined.Broadcast(Event);
@@ -1007,27 +1077,6 @@ TOnlineAsyncOpHandle<FJoinSession> FSessionsEOSGS::JoinSession(FJoinSession::Par
 	return Op->GetHandle();
 }
 
-TResult<TSharedRef<const FSessionInvite>, FOnlineError> FSessionsEOSGS::BuildSessionInvite(FOnlineAccountIdHandle RecipientId, FOnlineAccountIdHandle SenderId, const FString& InInviteId) const
-{
-	TResult<TSharedRef<const FSession>, FOnlineError> SessionInviteResult = BuildSessionFromInvite(InInviteId);
-
-	if (SessionInviteResult.IsOk())
-	{
-		FSessionInvite SessionInvite {
-			RecipientId,
-			SenderId,
-			InInviteId,
-			SessionInviteResult.GetOkValue()
-		};
-
-		return TResult<TSharedRef<const FSessionInvite>, FOnlineError>(MakeShared<FSessionInvite>(MoveTemp(SessionInvite)));
-	}
-	else
-	{
-		return TResult<TSharedRef<const FSessionInvite>, FOnlineError>(SessionInviteResult.GetErrorValue());
-	}
-}
-
 TResult<TSharedRef<const FSession>, FOnlineError> FSessionsEOSGS::BuildSessionFromInvite(const FString& InInviteId) const
 {
 	EOS_Sessions_CopySessionHandleByInviteIdOptions CopySessionHandleByInviteIdOptions = { };
@@ -1041,7 +1090,7 @@ TResult<TSharedRef<const FSession>, FOnlineError> FSessionsEOSGS::BuildSessionFr
 	EOS_EResult CopySessionHandleByInviteIdResult = EOS_Sessions_CopySessionHandleByInviteId(SessionsHandle, &CopySessionHandleByInviteIdOptions, &SessionDetailsHandle);
 	if (CopySessionHandleByInviteIdResult == EOS_EResult::EOS_Success)
 	{
-		return TResult<TSharedRef<const FSession>, FOnlineError>(MakeShared<FSessionEOSGS>(SessionDetailsHandle));
+		return TResult<TSharedRef<const FSession>, FOnlineError>(MakeShared<const FSessionEOSGS>(SessionDetailsHandle));
 	}
 	else
 	{
@@ -1140,7 +1189,8 @@ TOnlineAsyncOpHandle<FRejectSessionInvite> FSessionsEOSGS::RejectSessionInvite(F
 
 		RejectInviteOptions.LocalUserId = GetProductUserIdChecked(OpParams.LocalUserId);
 
-		const FTCHARToUTF8 InviteIdUtf8(OpParams.SessionInvite->InviteId);
+		const FString& InviteIdStr = FOnlineSessionInviteIdRegistryEOSGS::Get().BasicRegistry.FindIdValue(OpParams.SessionInviteId);
+		const FTCHARToUTF8 InviteIdUtf8(InviteIdStr);
 		RejectInviteOptions.InviteId = InviteIdUtf8.Get();
 
 		EOS_Async(EOS_Sessions_RejectInvite, SessionsHandle, RejectInviteOptions, MoveTemp(Promise));
@@ -1154,6 +1204,13 @@ TOnlineAsyncOpHandle<FRejectSessionInvite> FSessionsEOSGS::RejectSessionInvite(F
 			return;
 		}
 
+		const FRejectSessionInvite::Params& OpParams = Op.GetParams();
+
+		if (TMap<FOnlineSessionInviteIdHandle, TSharedRef<FSessionInvite>>* UserMap = SessionInvitesUserMap.Find(OpParams.LocalUserId))
+		{
+			UserMap->Remove(OpParams.SessionInviteId);
+		}
+
 		Op.SetResult(FRejectSessionInvite::Result{ });
 	})
 	.Enqueue(GetSerialQueue()); // TODO: Use the parallel queue instead when possible
@@ -1299,19 +1356,6 @@ TOnlineAsyncOpHandle<FRemoveSessionMembers> FSessionsEOSGS::RemoveSessionMembers
 	return Op->GetHandle();
 }
 
-bool FSessionsEOSGS::TryRemoveSession(const FName& SessionName)
-{
-	if (TSharedRef<FSession>* FoundSession = SessionsByName.Find(SessionName))
-	{
-		SessionsById.Remove((*FoundSession)->SessionId);
-		SessionsByName.Remove(SessionName);
-
-		return true;
-	}
-
-	return false;
-}
-
 /** FSessionsLAN */
 
 void FSessionsEOSGS::AppendSessionToPacket(FNboSerializeToBuffer& Packet, const FSessionLAN& Session)
diff --git a/Engine/Plugins/Online/Experimental/OnlineServicesEOSGS/Source/Public/Online/SessionsEOSGS.h b/Engine/Plugins/Online/Experimental/OnlineServicesEOSGS/Source/Public/Online/SessionsEOSGS.h
--- a/Engine/Plugins/Online/Experimental/OnlineServicesEOSGS/Source/Public/Online/SessionsEOSGS.h
+++ b/Engine/Plugins/Online/Experimental/OnlineServicesEOSGS/Source/Public/Online/SessionsEOSGS.h
@@ -20,6 +20,16 @@ class FOnlineSessionIdRegistryEOSGS : public TOnlineSessionIdStringRegistry<EOnl
 {
 public:
 	static FOnlineSessionIdRegistryEOSGS& Get();
+
+	bool IsSessionIdExpired(const FOnlineSessionIdHandle& InHandle) const;
+};
+
+class FOnlineSessionInviteIdRegistryEOSGS : public TOnlineSessionInviteIdStringRegistry<EOnlineServices::Epic>
+{
+public:
+	static FOnlineSessionInviteIdRegistryEOSGS& Get();
+
+	bool IsSessionInviteIdExpired(const FOnlineSessionInviteIdHandle& InHandle) const;
 };
 
 static FName EOS_SESSIONS_BUCKET_ID = TEXT("EOS_SESSIONS_BUCKET_ID");
@@ -152,14 +162,10 @@ protected:
 	void WriteSessionSearchHandle(FSessionSearchHandleEOSGS& SessionSearchHandle, const FFindSessions::Params& Params);
 
 	static FOnlineSessionIdHandle CreateSessionId(const FString& SessionId);
+	FOnlineSessionInviteIdHandle CreateSessionInviteId(const FString& SessionInviteId) const;
 
 	/**
-	 * Builds a session invite from an invite id, using the class' Sessions Handle
-	 */
-	TResult<TSharedRef<const FSessionInvite>, FOnlineError> BuildSessionInvite(FOnlineAccountIdHandle RecipientId, FOnlineAccountIdHandle SenderId, const FString& InInviteId) const;
-
-	/**
-	* Builds a session from an invite id, using the class' Sessions Handle
+	 * Builds a session from an invite id, using the class' Sessions Handle
 	 */
 	TResult<TSharedRef<const FSession>, FOnlineError> BuildSessionFromInvite(const FString& InInviteId) const;
 
@@ -169,9 +175,6 @@ protected:
 	TResult<TSharedRef<const FSession>, FOnlineError> BuildSessionFromUIEvent(const EOS_UI_EventId& UIEventId) const;
 
 private:
-	/** Returns true if the named session is found and removed successfully, and false otherwise */
-	bool TryRemoveSession(const FName& SessionName);
-
 	// FSessionsLAN
 	virtual void AppendSessionToPacket(FNboSerializeToBuffer& Packet, const FSessionLAN& Session) override;
 	virtual void ReadSessionFromPacket(FNboSerializeFromBuffer& Packet, FSessionLAN& Session) override;
diff --git a/Engine/Source/Runtime/CoreOnline/Public/Online/CoreOnline.h b/Engine/Source/Runtime/CoreOnline/Public/Online/CoreOnline.h
--- a/Engine/Source/Runtime/CoreOnline/Public/Online/CoreOnline.h
+++ b/Engine/Source/Runtime/CoreOnline/Public/Online/CoreOnline.h
@@ -211,6 +211,7 @@ namespace OnlineIdHandleTags
 {
 	struct FAccount {};
 	struct FSession {};
+	struct FSessionInvite {};
 	struct FLobby {};
 	struct FVerifiedAuthTicket {};
 	struct FVerifiedAuthSession {};
@@ -308,6 +309,7 @@ private:
 using FOnlineAccountIdHandle = TOnlineIdHandle<OnlineIdHandleTags::FAccount>;
 using FOnlineLobbyIdHandle = TOnlineIdHandle<OnlineIdHandleTags::FLobby>;
 using FOnlineSessionIdHandle = TOnlineIdHandle<OnlineIdHandleTags::FSession>;
+using FOnlineSessionInviteIdHandle = TOnlineIdHandle<OnlineIdHandleTags::FSessionInvite>;
 using FOnlineVerifiedAuthTicketIdHandle = TOnlineIdHandle<OnlineIdHandleTags::FVerifiedAuthTicket>;
 using FOnlineVerifiedAuthSessionIdHandle = TOnlineIdHandle<OnlineIdHandleTags::FVerifiedAuthSession>;
 
@@ -335,6 +337,7 @@ public:
 
 using IOnlineAccountIdRegistry = IOnlineIdRegistry<OnlineIdHandleTags::FAccount>;
 using IOnlineSessionIdRegistry = IOnlineIdRegistry<OnlineIdHandleTags::FSession>;
+using IOnlineSessionInviteIdRegistry = IOnlineIdRegistry<OnlineIdHandleTags::FSessionInvite>;
 
 class FOnlineIdRegistryRegistry
 {
