Update OnlineServicesEOSGS set reserve options to point to gamedev using config boolean.

diff --git a/Engine/Plugins/Online/Experimental/OnlineServicesEOSGS/Source/Private/Online/OnlineServicesEOSGSPlatformFactory.cpp b/Engine/Plugins/Online/Experimental/OnlineServicesEOSGS/Source/Private/Online/OnlineServicesEOSGSPlatformFactory.cpp
--- a/Engine/Plugins/Online/Experimental/OnlineServicesEOSGS/Source/Private/Online/OnlineServicesEOSGSPlatformFactory.cpp
+++ b/Engine/Plugins/Online/Experimental/OnlineServicesEOSGS/Source/Private/Online/OnlineServicesEOSGSPlatformFactory.cpp
@@ -22,6 +22,7 @@ struct FEOSPlatformConfig
 	FString ClientId;
 	FString ClientSecret;
 	FString EncryptionKey;
+	bool bUseGameDev;
 };
 
 FEOSPlatformConfig LoadEOSPlatformConfig()
@@ -34,6 +35,7 @@ FEOSPlatformConfig LoadEOSPlatformConfig()
 	GConfig->GetString(*ConfigSection, TEXT("ClientId"), PlatformConfig.ClientId, GEngineIni);
 	GConfig->GetString(*ConfigSection, TEXT("ClientSecret"), PlatformConfig.ClientSecret, GEngineIni);
 	GConfig->GetString(*ConfigSection, TEXT("EncryptionKey"), PlatformConfig.EncryptionKey, GEngineIni);
+	GConfig->GetBool(*ConfigSection, TEXT("bUseGameDev"), PlatformConfig.bUseGameDev, GEngineIni);
 	return PlatformConfig;
 }
 
@@ -96,9 +98,15 @@ IEOSPlatformHandlePtr FOnlineServicesEOSGSPlatformFactory::CreatePlatform()
 	const FTCHARToUTF8 EncryptionKey(EOSPlatformConfig.EncryptionKey.IsEmpty() ? nullptr : *EOSPlatformConfig.EncryptionKey);
 	const FTCHARToUTF8 CacheDirectory(*(SDKManager->GetCacheDirBase() / TEXT("OnlineServicesEOS")));
 
+	static struct FReservedOptions
+	{
+		int32_t ApiVersion;
+		const char* BackendEnvironment;
+	} ReservedOptions = { 1, "GameDev" };
+
 	EOS_Platform_Options PlatformOptions = {};
 	PlatformOptions.ApiVersion = EOS_PLATFORM_OPTIONS_API_LATEST;
-	PlatformOptions.Reserved = nullptr;
+	PlatformOptions.Reserved = EOSPlatformConfig.bUseGameDev ? &ReservedOptions : nullptr;
 	PlatformOptions.bIsServer = EOS_FALSE;
 	PlatformOptions.OverrideCountryCode = nullptr;
 	PlatformOptions.OverrideLocaleCode = nullptr;
