- Catch "NotFound" errors from EnumerateFiles, and return success instead. NotFound is EOS's way of saying the operation worked, but the user has no files.
- Fix Read and Write implementations crashing, all 3 callbacks (Data/Transfer/Complete) share a single client data.
- Fixed a few cases of FTCHARToUTF8 taking a FString not a TCHAR*, which compiles but the converted string is not null terminated.

diff --git a/Engine/Plugins/Online/Experimental/OnlineServicesEOSGS/Source/Private/Online/UserFileEOSGS.cpp b/Engine/Plugins/Online/Experimental/OnlineServicesEOSGS/Source/Private/Online/UserFileEOSGS.cpp
--- a/Engine/Plugins/Online/Experimental/OnlineServicesEOSGS/Source/Private/Online/UserFileEOSGS.cpp
+++ b/Engine/Plugins/Online/Experimental/OnlineServicesEOSGS/Source/Private/Online/UserFileEOSGS.cpp
@@ -62,7 +62,8 @@ TOnlineAsyncOpHandle<FUserFileEnumerateFiles> FUserFileEOSGS::EnumerateFiles(FUs
 	{
 		const FUserFileEnumerateFiles::Params& Params = Op.GetParams();
 
-		if (Data->ResultCode != EOS_EResult::EOS_Success)
+		if (Data->ResultCode != EOS_EResult::EOS_Success
+			&& Data->ResultCode != EOS_EResult::EOS_NotFound)
 		{
 			UE_LOG(LogTemp, Warning, TEXT("EOS_PlayerDataStorage_QueryFileList failed with result=[%s]"), *LexToString(Data->ResultCode));
 			Op.SetError(Errors::FromEOSResult(Data->ResultCode));
@@ -136,6 +137,19 @@ TOnlineResult<FUserFileGetEnumeratedFiles> FUserFileEOSGS::GetEnumeratedFiles(FU
 	return TOnlineResult<FUserFileGetEnumeratedFiles>({*Found});
 }
 
+struct FUserFileReadFileClientData
+{
+	using PromiseType = TPromise<const EOS_PlayerDataStorage_ReadFileCallbackInfo*>;
+
+	FUserFileReadFileClientData(const FUserFileContentsRef& InFileContents, PromiseType&& InPromise)
+		: FileContents(InFileContents)
+		, Promise(MoveTemp(InPromise))
+	{}
+
+	FUserFileContentsWeakPtr FileContents;
+	PromiseType Promise;
+};
+
 TOnlineAsyncOpHandle<FUserFileReadFile> FUserFileEOSGS::ReadFile(FUserFileReadFile::Params&& InParams)
 {
 	static const TCHAR* FileContentsKey = TEXT("FileContents");
@@ -163,33 +177,20 @@ TOnlineAsyncOpHandle<FUserFileReadFile> FUserFileEOSGS::ReadFile(FUserFileReadFi
 		FUserFileContentsRef FileContentsRef = MakeShared<FUserFileContents>();
 		Op.Data.Set<FUserFileContentsRef>(FileContentsKey, FileContentsRef);
 
-		using FReadFileDataCallback = TEOSGlobalCallback<EOS_PlayerDataStorage_OnReadFileDataCallback, EOS_PlayerDataStorage_ReadFileDataCallbackInfo, TOnlineAsyncOp<FUserFileReadFile>, EOS_PlayerDataStorage_EReadResult>;
-		TSharedRef<FReadFileDataCallback> ReadFileDataCallback = MakeShared<FReadFileDataCallback>(Op.AsWeak());
-		Op.Data.Set<TSharedRef<FReadFileDataCallback>>(TEXT("ReadFileDataCallback"), ReadFileDataCallback);
-		ReadFileDataCallback->CallbackLambda = [FileContentsRef](const EOS_PlayerDataStorage_ReadFileDataCallbackInfo* Data) -> EOS_PlayerDataStorage_EReadResult
-		{
-			FUserFileContents& FileContents = const_cast<FUserFileContents&>(FileContentsRef.Get());
-			// Only has any effect on the first callback
-			FileContents.Reserve(Data->TotalFileSizeBytes);
-			FileContents.Append((uint8*)Data->DataChunk, Data->DataChunkLengthBytes);
-
-			UE_LOG(LogTemp, VeryVerbose, TEXT("FUserFileEOSGS::ReadFile ReadProgress Filename=[%s] %d/%d"), UTF8_TO_TCHAR(Data->Filename), FileContents.Num(), Data->TotalFileSizeBytes);
-
-			return EOS_PlayerDataStorage_EReadResult::EOS_RR_ContinueReading;
-		};
-
-		const FTCHARToUTF8 Utf8Filename(Params.Filename);
+		const auto Utf8Filename = StringCast<UTF8CHAR>(*Params.Filename);
 
-		EOS_PlayerDataStorage_ReadFileOptions  Options = {};
+		EOS_PlayerDataStorage_ReadFileOptions Options = {};
 		Options.ApiVersion = EOS_PLAYERDATASTORAGE_READFILEOPTIONS_API_LATEST;
 		static_assert(EOS_PLAYERDATASTORAGE_READFILEOPTIONS_API_LATEST == 1, "EOS_PlayerDataStorage_ReadFileOptions updated, check new fields");
 		Options.LocalUserId = GetProductUserIdChecked(Params.LocalUserId);
-		Options.Filename = Utf8Filename.Get();
+		Options.Filename = (const char*)Utf8Filename.Get();
 		Options.ReadChunkLengthBytes = Config.ChunkLengthBytes;
-		Options.ReadFileDataCallback = ReadFileDataCallback->GetCallbackPtr();
+		Options.ReadFileDataCallback = &FUserFileEOSGS::OnReadFileDataStatic;
 		Options.FileTransferProgressCallback = &FUserFileEOSGS::OnFileTransferProgressStatic;
 
-		const EOS_HPlayerDataStorageFileTransferRequest RequestHandle = EOS_Async(EOS_PlayerDataStorage_ReadFile, PlayerDataStorageHandle, Options, MoveTemp(Promise));
+		FUserFileReadFileClientData* ClientData = new FUserFileReadFileClientData(FileContentsRef, MoveTemp(Promise));
+
+		const EOS_HPlayerDataStorageFileTransferRequest RequestHandle = EOS_PlayerDataStorage_ReadFile(PlayerDataStorageHandle, &Options, ClientData, &FUserFileEOSGS::OnReadFileCompleteStatic);
 		Op.Data.Set<EOS_HPlayerDataStorageFileTransferRequest>(RequestHandleKey, RequestHandle);
 	})
 	.Then([this](TOnlineAsyncOp<FUserFileReadFile>& Op, const EOS_PlayerDataStorage_ReadFileCallbackInfo* Data)
@@ -221,6 +222,20 @@ TOnlineAsyncOpHandle<FUserFileReadFile> FUserFileEOSGS::ReadFile(FUserFileReadFi
 	return Op->GetHandle();
 }
 
+struct FUserFileWriteFileClientData
+{
+	using PromiseType = TPromise<const EOS_PlayerDataStorage_WriteFileCallbackInfo*>;
+
+	FUserFileWriteFileClientData(FUserFileContents&& InFileContents, PromiseType&& InPromise)
+		: FileContents(MoveTemp(InFileContents))
+		, Promise(MoveTemp(InPromise))
+	{}
+
+	uint32_t BytesWritten = 0;
+	FUserFileContents FileContents;
+	PromiseType Promise;
+};
+
 TOnlineAsyncOpHandle<FUserFileWriteFile> FUserFileEOSGS::WriteFile(FUserFileWriteFile::Params&& InParams)
 {
 	TOnlineAsyncOpRef<FUserFileWriteFile> Op = GetOp<FUserFileWriteFile>(MoveTemp(InParams));
@@ -243,58 +258,20 @@ TOnlineAsyncOpHandle<FUserFileWriteFile> FUserFileEOSGS::WriteFile(FUserFileWrit
 			return;
 		}
 
-		using FWriteFileDataCallback = TEOSGlobalCallback<EOS_PlayerDataStorage_OnWriteFileDataCallback, EOS_PlayerDataStorage_WriteFileDataCallbackInfo, TOnlineAsyncOp<FUserFileWriteFile>, EOS_PlayerDataStorage_EWriteResult, void*, uint32_t*>;
-		TSharedRef<FWriteFileDataCallback> WriteFileDataCallback = MakeShared<FWriteFileDataCallback>(Op.AsWeak());
-		Op.Data.Set<TSharedRef<FWriteFileDataCallback>>(TEXT("WriteFileDataCallback"), WriteFileDataCallback);
-		WriteFileDataCallback->CallbackLambda = [WeakOp = Op.AsWeak()] (const EOS_PlayerDataStorage_WriteFileDataCallbackInfo* CallbackInfo, void* OutDataBuffer, uint32_t* OutDataWritten)
-		{
-			TOnlineAsyncOpPtr<FUserFileWriteFile> Op = WeakOp.Pin();
-			if (!ensure(Op))
-			{
-				return EOS_PlayerDataStorage_EWriteResult::EOS_WR_FailRequest;
-			}
-
-			const TCHAR* BytesWrittenKey = TEXT("BytesWritten");
-
-			const FUserFileWriteFile::Params& Params = Op->GetParams();
-			const FUserFileContents& FileContents = Params.FileContents;
-
-			const uint32_t* BytesWrittenPtr = Op->Data.Get<uint32_t>(BytesWrittenKey);
-			const uint32_t BytesWrittenBefore = BytesWrittenPtr ? *BytesWrittenPtr : 0;
-			const uint32_t BytesRemaining = FileContents.Num() - BytesWrittenBefore;
-			const uint32_t BytesToWrite = FMath::Min(CallbackInfo->DataBufferLengthBytes, BytesRemaining);
-
-			FMemory::Memcpy(OutDataBuffer, FileContents.GetData() + BytesWrittenBefore, BytesToWrite);
-			*OutDataWritten = BytesToWrite;
-
-			const uint32_t BytesWrittenTotal = BytesWrittenBefore + BytesToWrite;
-			Op->Data.Set<uint32_t>(BytesWrittenKey, BytesWrittenTotal);
-
-			UE_LOG(LogTemp, VeryVerbose, TEXT("FUserFileEOSGS::WriteFile WriteProgress Filename=[%s] %d/%d"), UTF8_TO_TCHAR(CallbackInfo->Filename), BytesWrittenTotal, FileContents.Num());
-
-			const bool bDone = BytesWrittenTotal == FileContents.Num();
-			if (bDone)
-			{
-				return EOS_PlayerDataStorage_EWriteResult::EOS_WR_CompleteRequest;
-			}
-			else
-			{
-				return EOS_PlayerDataStorage_EWriteResult::EOS_WR_ContinueWriting;
-			}
-		};
-
-		const FTCHARToUTF8 Utf8Filename(Params.Filename);
+		const auto Utf8Filename = StringCast<UTF8CHAR>(*Params.Filename);
 
-		EOS_PlayerDataStorage_WriteFileOptions  Options = {};
+		EOS_PlayerDataStorage_WriteFileOptions Options = {};
 		Options.ApiVersion = EOS_PLAYERDATASTORAGE_WRITEFILEOPTIONS_API_LATEST;
 		static_assert(EOS_PLAYERDATASTORAGE_WRITEFILEOPTIONS_API_LATEST == 1, "EOS_PlayerDataStorage_WriteFileOptions updated, check new fields");
 		Options.LocalUserId = GetProductUserIdChecked(Params.LocalUserId);
-		Options.Filename = Utf8Filename.Get();
+		Options.Filename = (const char*)Utf8Filename.Get();
 		Options.ChunkLengthBytes = Config.ChunkLengthBytes;
-		Options.WriteFileDataCallback = WriteFileDataCallback->GetCallbackPtr();
+		Options.WriteFileDataCallback = &FUserFileEOSGS::OnWriteFileDataStatic;
 		Options.FileTransferProgressCallback = &FUserFileEOSGS::OnFileTransferProgressStatic;
 
-		const EOS_HPlayerDataStorageFileTransferRequest RequestHandle = EOS_Async(EOS_PlayerDataStorage_WriteFile, PlayerDataStorageHandle, Options, MoveTemp(Promise));
+		FUserFileWriteFileClientData* ClientData = new FUserFileWriteFileClientData(MoveTemp(const_cast<FUserFileWriteFile::Params&>(Params).FileContents), MoveTemp(Promise));
+
+		const EOS_HPlayerDataStorageFileTransferRequest RequestHandle = EOS_PlayerDataStorage_WriteFile(PlayerDataStorageHandle, &Options, ClientData, &FUserFileEOSGS::OnWriteFileCompleteStatic);
 		Op.Data.Set<EOS_HPlayerDataStorageFileTransferRequest>(RequestHandleKey, RequestHandle);
 	})
 	.Then([this](TOnlineAsyncOp<FUserFileWriteFile>& Op, const EOS_PlayerDataStorage_WriteFileCallbackInfo* Data)
@@ -340,15 +317,15 @@ TOnlineAsyncOpHandle<FUserFileCopyFile> FUserFileEOSGS::CopyFile(FUserFileCopyFi
 			return;
 		}
 
-		const FTCHARToUTF8 Utf8SourceFilename(Params.SourceFilename);
-		const FTCHARToUTF8 Utf8TargetFilename(Params.TargetFilename);
+		const auto Utf8SourceFilename = StringCast<UTF8CHAR>(*Params.SourceFilename);
+		const auto Utf8TargetFilename = StringCast<UTF8CHAR>(*Params.TargetFilename);
 
-		EOS_PlayerDataStorage_DuplicateFileOptions  Options = {};
+		EOS_PlayerDataStorage_DuplicateFileOptions Options = {};
 		Options.ApiVersion = EOS_PLAYERDATASTORAGE_DUPLICATEFILEOPTIONS_API_LATEST;
 		static_assert(EOS_PLAYERDATASTORAGE_DUPLICATEFILEOPTIONS_API_LATEST == 1, "EOS_PlayerDataStorage_DuplicateFileOptions updated, check new fields");
 		Options.LocalUserId = GetProductUserIdChecked(Params.LocalUserId);
-		Options.SourceFilename = Utf8SourceFilename.Get();
-		Options.DestinationFilename = Utf8TargetFilename.Get();
+		Options.SourceFilename = (const char*)Utf8SourceFilename.Get();
+		Options.DestinationFilename = (const char*)Utf8TargetFilename.Get();
 
 		EOS_Async(EOS_PlayerDataStorage_DuplicateFile, PlayerDataStorageHandle, Options, MoveTemp(Promise));
 	})
@@ -390,13 +367,13 @@ TOnlineAsyncOpHandle<FUserFileDeleteFile> FUserFileEOSGS::DeleteFile(FUserFileDe
 			return;
 		}
 
-		const FTCHARToUTF8 Utf8Filename(Params.Filename);
+		const auto Utf8Filename = StringCast<UTF8CHAR>(*Params.Filename);
 
-		EOS_PlayerDataStorage_DeleteFileOptions  Options = {};
+		EOS_PlayerDataStorage_DeleteFileOptions Options = {};
 		Options.ApiVersion = EOS_PLAYERDATASTORAGE_DELETEFILEOPTIONS_API_LATEST;
 		static_assert(EOS_PLAYERDATASTORAGE_DELETEFILEOPTIONS_API_LATEST == 1, "EOS_PlayerDataStorage_DeleteFileOptions updated, check new fields");
 		Options.LocalUserId = GetProductUserIdChecked(Params.LocalUserId);
-		Options.Filename = Utf8Filename.Get();
+		Options.Filename = (const char*)Utf8Filename.Get();
 
 		EOS_Async(EOS_PlayerDataStorage_DeleteFile, PlayerDataStorageHandle, Options, MoveTemp(Promise));
 	})
@@ -416,10 +393,79 @@ TOnlineAsyncOpHandle<FUserFileDeleteFile> FUserFileEOSGS::DeleteFile(FUserFileDe
 	return Op->GetHandle();
 }
 
+EOS_PlayerDataStorage_EReadResult EOS_CALL FUserFileEOSGS::OnReadFileDataStatic(const EOS_PlayerDataStorage_ReadFileDataCallbackInfo* Data)
+{
+	FUserFileReadFileClientData* ClientData = (FUserFileReadFileClientData*)Data->ClientData;
+	check(ClientData);
+	FUserFileContentsPtr FileContentsPtr = ClientData->FileContents.Pin();
+	if (ensure(FileContentsPtr))
+	{
+		FUserFileContents& FileContents = *ConstCastSharedPtr<FUserFileContents>(FileContentsPtr);
+		// Only has any effect on the first callback
+		FileContents.Reserve(Data->TotalFileSizeBytes);
+		FileContents.Append((uint8*)Data->DataChunk, Data->DataChunkLengthBytes);
+
+		UE_LOG(LogTemp, VeryVerbose, TEXT("EOS_PlayerDataStorage_ReadFile Progress Filename=[%s] %d/%d"), UTF8_TO_TCHAR(Data->Filename), FileContents.Num(), Data->TotalFileSizeBytes);
+	}
+
+	return EOS_PlayerDataStorage_EReadResult::EOS_RR_ContinueReading;
+}
+
+void EOS_CALL FUserFileEOSGS::OnReadFileCompleteStatic(const EOS_PlayerDataStorage_ReadFileCallbackInfo* Data)
+{
+	UE_LOG(LogTemp, VeryVerbose, TEXT("EOS_PlayerDataStorage_ReadFile Complete Filename=[%s]"), UTF8_TO_TCHAR(Data->Filename));
+
+	FUserFileReadFileClientData* ClientData = (FUserFileReadFileClientData*)Data->ClientData;
+	check(ClientData);
+	ClientData->Promise.EmplaceValue(Data);
+
+	delete ClientData;
+}
+
+EOS_PlayerDataStorage_EWriteResult EOS_CALL FUserFileEOSGS::OnWriteFileDataStatic(const EOS_PlayerDataStorage_WriteFileDataCallbackInfo* Data, void* OutDataBuffer, uint32_t* OutDataWritten)
+{
+	FUserFileWriteFileClientData* ClientData = (FUserFileWriteFileClientData*)Data->ClientData;
+	check(ClientData);
+	
+	const FUserFileContents& FileContents = ClientData->FileContents;
+
+	const uint32_t BytesWrittenBefore = ClientData->BytesWritten;
+	const uint32_t BytesRemaining = FileContents.Num() - BytesWrittenBefore;
+	const uint32_t BytesToWrite = FMath::Min(Data->DataBufferLengthBytes, BytesRemaining);
+
+	FMemory::Memcpy(OutDataBuffer, FileContents.GetData() + BytesWrittenBefore, BytesToWrite);
+	*OutDataWritten = BytesToWrite;
+
+	const uint32_t BytesWrittenTotal = BytesWrittenBefore + BytesToWrite;
+	ClientData->BytesWritten = BytesWrittenTotal;
+
+	UE_LOG(LogTemp, VeryVerbose, TEXT("EOS_PlayerDataStorage_WriteFile Progress Filename=[%s] %d/%d"), UTF8_TO_TCHAR(Data->Filename), BytesWrittenTotal, FileContents.Num());
+
+	const bool bDone = BytesWrittenTotal == FileContents.Num();
+	if (bDone)
+	{
+		return EOS_PlayerDataStorage_EWriteResult::EOS_WR_CompleteRequest;
+	}
+	else
+	{
+		return EOS_PlayerDataStorage_EWriteResult::EOS_WR_ContinueWriting;
+	}
+}
+
+void EOS_CALL FUserFileEOSGS::OnWriteFileCompleteStatic(const EOS_PlayerDataStorage_WriteFileCallbackInfo* Data)
+{
+	UE_LOG(LogTemp, VeryVerbose, TEXT("EOS_PlayerDataStorage_WriteFile Complete Filename=[%s]"), UTF8_TO_TCHAR(Data->Filename));
+
+	FUserFileWriteFileClientData* ClientData = (FUserFileWriteFileClientData*)Data->ClientData;
+	check(ClientData);
+	ClientData->Promise.EmplaceValue(Data);
+
+	delete ClientData;
+}
 
 void EOS_CALL FUserFileEOSGS::OnFileTransferProgressStatic(const EOS_PlayerDataStorage_FileTransferProgressCallbackInfo* Data)
 {
-	UE_LOG(LogTemp, VeryVerbose, TEXT("FUserFileEOSGS::ReadFile TransferProgress Filename=[%s] %d/%d"), UTF8_TO_TCHAR(Data->Filename), Data->BytesTransferred, Data->TotalFileSizeBytes);
+	UE_LOG(LogTemp, VeryVerbose, TEXT("FUserFileEOSGS TransferProgress Filename=[%s] %d/%d"), UTF8_TO_TCHAR(Data->Filename), Data->BytesTransferred, Data->TotalFileSizeBytes);
 }
 
 /* UE::Online */ }
diff --git a/Engine/Plugins/Online/Experimental/OnlineServicesEOSGS/Source/Public/Online/UserFileEOSGS.h b/Engine/Plugins/Online/Experimental/OnlineServicesEOSGS/Source/Public/Online/UserFileEOSGS.h
--- a/Engine/Plugins/Online/Experimental/OnlineServicesEOSGS/Source/Public/Online/UserFileEOSGS.h
+++ b/Engine/Plugins/Online/Experimental/OnlineServicesEOSGS/Source/Public/Online/UserFileEOSGS.h
@@ -54,6 +54,12 @@ protected:
 
 	TMap<FOnlineAccountIdHandle, TArray<FString>> UserToFiles;
 
+	static EOS_PlayerDataStorage_EReadResult EOS_CALL OnReadFileDataStatic(const EOS_PlayerDataStorage_ReadFileDataCallbackInfo* Data);
+	static void EOS_CALL OnReadFileCompleteStatic(const EOS_PlayerDataStorage_ReadFileCallbackInfo* Data);
+
+	static EOS_PlayerDataStorage_EWriteResult EOS_CALL OnWriteFileDataStatic(const EOS_PlayerDataStorage_WriteFileDataCallbackInfo* Data, void* OutDataBuffer, uint32_t* OutDataWritten);
+	static void EOS_CALL OnWriteFileCompleteStatic(const EOS_PlayerDataStorage_WriteFileCallbackInfo* Data);
+
 	static void EOS_CALL OnFileTransferProgressStatic(const EOS_PlayerDataStorage_FileTransferProgressCallbackInfo* Data);
 };
 
