Added SchemaVariant support for all implementations of Sessions Interface
- Moved common Schema functionality out of Lobbies.
- Added serialization methods for all Sessions Data structures.

diff --git a/Engine/Plugins/Online/Experimental/OnlineServices/Source/OnlineServicesCommon/OnlineServicesCommon.Build.cs b/Engine/Plugins/Online/Experimental/OnlineServices/Source/OnlineServicesCommon/OnlineServicesCommon.Build.cs
--- a/Engine/Plugins/Online/Experimental/OnlineServices/Source/OnlineServicesCommon/OnlineServicesCommon.Build.cs
+++ b/Engine/Plugins/Online/Experimental/OnlineServices/Source/OnlineServicesCommon/OnlineServicesCommon.Build.cs
@@ -11,6 +11,7 @@ public class OnlineServicesCommon : ModuleRules
 		PublicDependencyModuleNames.AddRange(
 			new string[] {
 				"CoreOnline",
+				"OnlineBase",
 				"OnlineServicesInterface"
 			}
 		);
diff --git a/Engine/Plugins/Online/Experimental/OnlineServices/Source/OnlineServicesCommon/Private/Online/NboSerializerCommonSvc.cpp b/Engine/Plugins/Online/Experimental/OnlineServices/Source/OnlineServicesCommon/Private/Online/NboSerializerCommonSvc.cpp
new file mode 100644
--- /dev/null
+++ b/Engine/Plugins/Online/Experimental/OnlineServices/Source/OnlineServicesCommon/Private/Online/NboSerializerCommonSvc.cpp
@@ -0,0 +1,228 @@
+// Copyright Epic Games, Inc. All Rights Reserved.
+
+#include "Online/NboSerializerCommonSvc.h"
+
+/**
+ * Serializes data in network byte order form into a buffer
+ */
+namespace UE::Online {
+
+namespace NboSerializerCommonSvc {
+
+/** SerializeToBuffer methods */
+
+void SerializeToBuffer(FNboSerializeToBuffer& Packet, const FRegisteredPlayer& RegisteredPlayer)
+{
+	Packet << RegisteredPlayer.bHasReservedSlot;
+	Packet << RegisteredPlayer.bIsInSession;
+}
+
+void SerializeToBuffer(FNboSerializeToBuffer& Packet, const FSchemaVariant& Data)
+{
+	Packet << (uint8)Data.VariantType;
+
+	switch (Data.VariantType)
+	{
+	case ESchemaAttributeType::Bool: Packet << Data.GetBoolean(); break;
+	case ESchemaAttributeType::Double: Packet << Data.GetDouble(); break;
+	case ESchemaAttributeType::Int64: Packet << (int32)Data.GetInt64(); break;
+	case ESchemaAttributeType::String: Packet << Data.GetString(); break;
+	}
+}
+
+void SerializeToBuffer(FNboSerializeToBuffer& Packet, const FCustomSessionSetting& CustomSessionSetting)
+{
+	SerializeToBuffer(Packet, CustomSessionSetting.Data);
+	Packet << CustomSessionSetting.ID;
+	Packet << (uint8)CustomSessionSetting.Visibility;
+}
+
+void SerializeToBuffer(FNboSerializeToBuffer& Packet, const FCustomSessionSettingsMap& CustomSessionSettingsMap)
+{
+	// First count the number of advertised keys
+	int32 NumAdvertisedProperties = 0;
+	for (const TPair<FSchemaAttributeId, FCustomSessionSetting>& Entry : CustomSessionSettingsMap)
+	{
+		const FCustomSessionSetting& Setting = Entry.Value;
+		if (Setting.Visibility == ESchemaAttributeVisibility::Public)
+		{
+			NumAdvertisedProperties++;
+		}
+	}
+
+	// Add the count of advertised keys and the data
+	Packet << NumAdvertisedProperties;
+
+	for (const TPair<FSchemaAttributeId, FCustomSessionSetting>& Entry : CustomSessionSettingsMap)
+	{
+		const FCustomSessionSetting& Setting = Entry.Value;
+		if (Setting.Visibility == ESchemaAttributeVisibility::Public)
+		{
+			Packet << Entry.Key;
+			SerializeToBuffer(Packet, Setting);
+		}
+	}
+}
+
+void SerializeToBuffer(FNboSerializeToBuffer& Packet, const FSessionMember& SessionMember)
+{
+	SerializeToBuffer(Packet, SessionMember.MemberSettings);
+}
+
+void SerializeToBuffer(FNboSerializeToBuffer& Packet, const FSessionSettings& SessionSettings)
+{
+	Packet << SessionSettings.bAllowNewMembers;
+	Packet << SessionSettings.bAllowSanctionedPlayers;
+	Packet << SessionSettings.bAllowUnregisteredPlayers;
+	Packet << SessionSettings.bAntiCheatProtected;
+	Packet << SessionSettings.bIsDedicatedServerSession;
+	Packet << SessionSettings.bIsLANSession;
+	Packet << SessionSettings.bPresenceEnabled;
+	SerializeToBuffer(Packet, SessionSettings.CustomSettings);
+	Packet << (uint8)SessionSettings.JoinPolicy;
+	Packet << SessionSettings.NumMaxPrivateConnections;;
+	Packet << SessionSettings.NumMaxPublicConnections;
+	Packet << SessionSettings.NumOpenPrivateConnections;
+	Packet << SessionSettings.NumOpenPublicConnections;
+	// SessionSettings.RegisteredPlayers will be serialized in implementations as user id types will vary
+	Packet << SessionSettings.SchemaName;
+	Packet << SessionSettings.SessionIdOverride;
+	// SessionSettings.SessionMembers will be serialized in implementations as user id types will vary
+}
+
+void SerializeToBuffer(FNboSerializeToBuffer& Packet, const FSession& Session)
+{
+	// Session.OwnerUserId will be serialized in implementations as user id types will vary
+	// Session.SessionId will be serialized in implementations as user id types will vary
+	// Session.CurrentState won't be serialized as it should always be Valid
+	SerializeToBuffer(Packet, Session.SessionSettings);
+}
+
+/** SerializeFromBuffer methods */
+
+void SerializeFromBuffer(FNboSerializeFromBuffer& Packet, FRegisteredPlayer& RegisteredPlayer)
+{
+	uint8 Read = 0;
+
+	Packet >> Read;
+	RegisteredPlayer.bHasReservedSlot = !!Read;
+	Packet >> Read;
+	RegisteredPlayer.bIsInSession = !!Read;
+}
+
+void SerializeFromBuffer(FNboSerializeFromBuffer& Packet, FSchemaVariant& Data)
+{
+	uint8 VariantType = 0;
+	Packet >> VariantType;
+	Data.VariantType = (ESchemaAttributeType)VariantType;
+
+	switch (Data.VariantType)
+	{
+	case ESchemaAttributeType::Bool:
+	{
+		uint8 Read;
+		Packet >> Read;
+		Data.Set(!!Read);
+		break;
+	}
+	case ESchemaAttributeType::Double:
+	{
+		double Read = 0.0;
+		Packet >> Read;
+		Data.Set(Read);
+		break;
+	}
+	case ESchemaAttributeType::Int64:
+	{
+		int32 Read = 0;
+		Packet >> Read;
+		Data.Set((int64)Read);
+		break;
+	}
+	case ESchemaAttributeType::String:
+	{
+		FString Read;
+		Packet >> Read;
+		Data.Set(Read);
+		break;
+	}
+	}
+}
+
+void SerializeFromBuffer(FNboSerializeFromBuffer& Packet, FCustomSessionSetting& CustomSessionSetting)
+{
+	SerializeFromBuffer(Packet, CustomSessionSetting.Data);
+	Packet >> CustomSessionSetting.ID;
+	uint8 VisibilityNum = 0;
+	Packet >> VisibilityNum;
+	CustomSessionSetting.Visibility = (ESchemaAttributeVisibility)VisibilityNum;
+}
+
+void SerializeFromBuffer(FNboSerializeFromBuffer& Packet, FCustomSessionSettingsMap& CustomSessionSettingsMap)
+{
+	int32 NumEntries = 0;
+	Packet >> NumEntries;
+
+	for (int32 Index = 0; Index < NumEntries; ++Index)
+	{
+		FSchemaAttributeId Key;
+		Packet >> Key;
+
+		FCustomSessionSetting Value;
+		SerializeFromBuffer(Packet, Value);
+
+		CustomSessionSettingsMap.Emplace(Key, Value);
+	}
+}
+
+void SerializeFromBuffer(FNboSerializeFromBuffer& Packet, FSessionMember& SessionMember)
+{
+	SerializeFromBuffer(Packet, SessionMember.MemberSettings);
+}
+
+void SerializeFromBuffer(FNboSerializeFromBuffer& Packet, FSessionSettings& SessionSettings)
+{
+	uint8 Read = 0;
+
+	// Read all the booleans as bytes
+	Packet >> Read;
+	SessionSettings.bAllowNewMembers = !!Read;
+	Packet >> Read;
+	SessionSettings.bAllowSanctionedPlayers = !!Read;
+	Packet >> Read;
+	SessionSettings.bAllowUnregisteredPlayers = !!Read;
+	Packet >> Read;
+	SessionSettings.bAntiCheatProtected = !!Read;
+	Packet >> Read;
+	SessionSettings.bIsDedicatedServerSession = !!Read;
+	Packet >> Read;
+	SessionSettings.bIsLANSession = !!Read;
+	Packet >> Read;
+	SessionSettings.bPresenceEnabled = !!Read;
+
+	SerializeFromBuffer(Packet, SessionSettings.CustomSettings);
+
+	Packet >> Read;
+	SessionSettings.JoinPolicy = (ESessionJoinPolicy)Read;
+
+	Packet >> SessionSettings.NumMaxPrivateConnections;;
+	Packet >> SessionSettings.NumMaxPublicConnections;
+	Packet >> SessionSettings.NumOpenPrivateConnections;
+	Packet >> SessionSettings.NumOpenPublicConnections;
+	// SessionSettings.RegisteredPlayers will be deserialized in implementations as user id types will vary
+	Packet >> SessionSettings.SchemaName;
+	Packet >> SessionSettings.SessionIdOverride;
+	// SessionSettings.SessionMembers will be deserialized in implementations as user id types will vary
+}
+
+void SerializeFromBuffer(FNboSerializeFromBuffer& Packet, FSession& Session)
+{
+	// Session.OwnerUserId will be deserialized in implementations as user id types will vary
+	// Session.SessionId will be deserialized in implementations as user id types will vary
+	Session.CurrentState = ESessionState::Valid;
+	SerializeFromBuffer(Packet, Session.SessionSettings);
+}
+
+/* NboSerializerCommonSvc */ }
+
+/* UE::Online */ }
\ No newline at end of file
diff --git a/Engine/Plugins/Online/Experimental/OnlineServices/Source/OnlineServicesCommon/Public/Online/NboSerializerCommonSvc.h b/Engine/Plugins/Online/Experimental/OnlineServices/Source/OnlineServicesCommon/Public/Online/NboSerializerCommonSvc.h
new file mode 100644
--- /dev/null
+++ b/Engine/Plugins/Online/Experimental/OnlineServices/Source/OnlineServicesCommon/Public/Online/NboSerializerCommonSvc.h
@@ -0,0 +1,36 @@
+// Copyright Epic Games, Inc. All Rights Reserved.
+
+#include "Online/NboSerializer.h"
+
+#include "Online/SessionsCommon.h"
+
+/**
+ * Serializes data in network byte order form into a buffer
+ */
+namespace UE::Online {
+
+namespace NboSerializerCommonSvc {
+
+/** SerializeToBuffer methods */
+
+ONLINESERVICESCOMMON_API void SerializeToBuffer(FNboSerializeToBuffer& Packet, const FRegisteredPlayer& RegisteredPlayer);
+ONLINESERVICESCOMMON_API void SerializeToBuffer(FNboSerializeToBuffer& Packet, const FSchemaVariant& Data);
+ONLINESERVICESCOMMON_API void SerializeToBuffer(FNboSerializeToBuffer& Packet, const FCustomSessionSetting& CustomSessionSetting);
+ONLINESERVICESCOMMON_API void SerializeToBuffer(FNboSerializeToBuffer& Packet, const FCustomSessionSettingsMap& CustomSessionSettingsMap);
+ONLINESERVICESCOMMON_API void SerializeToBuffer(FNboSerializeToBuffer& Packet, const FSessionMember& SessionMember);
+ONLINESERVICESCOMMON_API void SerializeToBuffer(FNboSerializeToBuffer& Packet, const FSessionSettings& SessionSettings);
+ONLINESERVICESCOMMON_API void SerializeToBuffer(FNboSerializeToBuffer& Packet, const FSession& Session);
+
+/** SerializeFromBuffer methods */
+
+ONLINESERVICESCOMMON_API void SerializeFromBuffer(FNboSerializeFromBuffer& Packet, FRegisteredPlayer& RegisteredPlayer);
+ONLINESERVICESCOMMON_API void SerializeFromBuffer(FNboSerializeFromBuffer& Packet, FSchemaVariant& Data);
+ONLINESERVICESCOMMON_API void SerializeFromBuffer(FNboSerializeFromBuffer& Packet, FCustomSessionSetting& CustomSessionSetting);
+ONLINESERVICESCOMMON_API void SerializeFromBuffer(FNboSerializeFromBuffer& Packet, FCustomSessionSettingsMap& CustomSessionSettingsMap);
+ONLINESERVICESCOMMON_API void SerializeFromBuffer(FNboSerializeFromBuffer& Packet, FSessionMember& SessionMember);
+ONLINESERVICESCOMMON_API void SerializeFromBuffer(FNboSerializeFromBuffer& Packet, FSessionSettings& SessionSettings);
+ONLINESERVICESCOMMON_API void SerializeFromBuffer(FNboSerializeFromBuffer& Packet, FSession& Session);
+
+/* NboSerializerCommonSvc */ }
+
+/* UE::Online */ }
\ No newline at end of file
diff --git a/Engine/Plugins/Online/Experimental/OnlineServices/Source/OnlineServicesCommonEngineUtils/Private/Online/SessionsLAN.cpp b/Engine/Plugins/Online/Experimental/OnlineServices/Source/OnlineServicesCommonEngineUtils/Private/Online/SessionsLAN.cpp
--- a/Engine/Plugins/Online/Experimental/OnlineServices/Source/OnlineServicesCommonEngineUtils/Private/Online/SessionsLAN.cpp
+++ b/Engine/Plugins/Online/Experimental/OnlineServices/Source/OnlineServicesCommonEngineUtils/Private/Online/SessionsLAN.cpp
@@ -7,6 +7,7 @@
 #include "Online/OnlineServicesCommonEngineUtils.h"
 #include "SocketSubsystem.h"
 #include "UObject/CoreNet.h"
+#include "Online/NboSerializerCommonSvc.h"
 
 namespace UE::Online {
 
@@ -29,6 +30,24 @@ FOnlineSessionIdHandle FOnlineSessionIdRegistryLAN::GetNextSessionId()
 	return BasicRegistry.FindOrAddHandle(FGuid::NewGuid().ToString());
 }
 
+/** NboSerializerLAN */
+
+namespace NboSerializerLANSvc {
+
+void SerializeToBuffer(FNboSerializeToBuffer& Packet, const FSessionLAN& Session)
+{
+	NboSerializerCommonSvc::SerializeToBuffer(Packet, Session);
+	Packet << *Session.OwnerInternetAddr;
+}
+
+void SerializeFromBuffer(FNboSerializeFromBuffer& Packet, FSessionLAN& Session)
+{
+	NboSerializerCommonSvc::SerializeFromBuffer(Packet, Session);
+	Packet >> *Session.OwnerInternetAddr;
+}
+
+/* NboSerializerLANSvc */ }
+
 /** FSessionLAN */
 
 FSessionLAN::FSessionLAN()
@@ -62,6 +81,12 @@ void FSessionLAN::Initialize()
 	{
 		OwnerInternetAddr->SetPort(7777); // Default port
 	}
+
+	// We'll set the connect address for the remote session as a custom parameter, so it can be read in OnlineServices' GetResolvedConnectString
+	FCustomSessionSetting ConnectString;
+	ConnectString.Data.Set(OwnerInternetAddr->ToString(true));
+	ConnectString.Visibility = ESchemaAttributeVisibility::Public;
+	SessionSettings.CustomSettings.Add(CONNECT_STRING_TAG, ConnectString);
 }
 
 FSessionLAN& FSessionLAN::Cast(FSession& InSession)
diff --git a/Engine/Plugins/Online/Experimental/OnlineServices/Source/OnlineServicesCommonEngineUtils/Public/Online/SessionsLAN.h b/Engine/Plugins/Online/Experimental/OnlineServices/Source/OnlineServicesCommonEngineUtils/Public/Online/SessionsLAN.h
--- a/Engine/Plugins/Online/Experimental/OnlineServices/Source/OnlineServicesCommonEngineUtils/Public/Online/SessionsLAN.h
+++ b/Engine/Plugins/Online/Experimental/OnlineServices/Source/OnlineServicesCommonEngineUtils/Public/Online/SessionsLAN.h
@@ -38,6 +38,13 @@ public:
 	TSharedPtr<FInternetAddr> OwnerInternetAddr;
 };
 
+namespace NboSerializerLANSvc {
+
+void ONLINESERVICESCOMMONENGINEUTILS_API SerializeToBuffer(FNboSerializeToBuffer& Packet, const FSessionLAN& Session);
+void ONLINESERVICESCOMMONENGINEUTILS_API SerializeFromBuffer(FNboSerializeFromBuffer& Packet, FSessionLAN& Session);
+
+/* NboSerializerLANSvc */ }
+
 class ONLINESERVICESCOMMONENGINEUTILS_API FSessionsLAN : public FSessionsCommon
 {
 public:
diff --git a/Engine/Plugins/Online/Experimental/OnlineServices/Source/OnlineServicesInterface/Private/Online/Lobbies.cpp b/Engine/Plugins/Online/Experimental/OnlineServices/Source/OnlineServicesInterface/Private/Online/Lobbies.cpp
--- a/Engine/Plugins/Online/Experimental/OnlineServices/Source/OnlineServicesInterface/Private/Online/Lobbies.cpp
+++ b/Engine/Plugins/Online/Experimental/OnlineServices/Source/OnlineServicesInterface/Private/Online/Lobbies.cpp
@@ -73,91 +73,6 @@ void LexFromString(ELobbyMemberLeaveReason& OutLeaveReason, const TCHAR* InStr)
 	}
 }
 
-const TCHAR* LexToString(ELobbyAttributeVisibility Visibility)
-{
-	switch (Visibility)
-	{
-	case ELobbyAttributeVisibility::Public:		return TEXT("Public");
-	default:									checkNoEntry(); // Intentional fallthrough
-	case ELobbyAttributeVisibility::Private:	return TEXT("Private");
-	}
-}
-
-void LexFromString(ELobbyAttributeVisibility& OutVisibility, const TCHAR* InStr)
-{
-	if (FCString::Stricmp(InStr, TEXT("Public")) == 0)
-	{
-		OutVisibility = ELobbyAttributeVisibility::Public;
-	}
-	else if (FCString::Stricmp(InStr, TEXT("Private")) == 0)
-	{
-		OutVisibility = ELobbyAttributeVisibility::Private;
-	}
-	else
-	{
-		checkNoEntry();
-		OutVisibility = ELobbyAttributeVisibility::Private;
-	}
-}
-
-const TCHAR* LexToString(ELobbyComparisonOp Comparison)
-{
-	switch (Comparison)
-	{
-	case ELobbyComparisonOp::Equals:			return TEXT("Equals");
-	case ELobbyComparisonOp::NotEquals:			return TEXT("NotEquals");
-	case ELobbyComparisonOp::GreaterThan:		return TEXT("GreaterThan");
-	case ELobbyComparisonOp::GreaterThanEquals:	return TEXT("GreaterThanEquals");
-	case ELobbyComparisonOp::LessThan:			return TEXT("LessThan");
-	case ELobbyComparisonOp::LessThanEquals:	return TEXT("LessThanEquals");
-	case ELobbyComparisonOp::Near:				return TEXT("Near");
-	case ELobbyComparisonOp::In:				return TEXT("In");
-	default:									checkNoEntry(); // Intentional fallthrough
-	case ELobbyComparisonOp::NotIn:				return TEXT("NotIn");
-	}
-}
-
-void LexFromString(ELobbyComparisonOp& OutComparison, const TCHAR* InStr)
-{
-	if (FCString::Stricmp(InStr, TEXT("Equals")) == 0)
-	{
-		OutComparison = ELobbyComparisonOp::Equals;
-	}
-	else if (FCString::Stricmp(InStr, TEXT("NotEquals")) == 0)
-	{
-		OutComparison = ELobbyComparisonOp::NotEquals;
-	}
-	else if (FCString::Stricmp(InStr, TEXT("GreaterThan")) == 0)
-	{
-		OutComparison = ELobbyComparisonOp::GreaterThan;
-	}
-	else if (FCString::Stricmp(InStr, TEXT("GreaterThanEquals")) == 0)
-	{
-		OutComparison = ELobbyComparisonOp::GreaterThanEquals;
-	}
-	else if (FCString::Stricmp(InStr, TEXT("LessThan")) == 0)
-	{
-		OutComparison = ELobbyComparisonOp::LessThan;
-	}
-	else if (FCString::Stricmp(InStr, TEXT("LessThanEquals")) == 0)
-	{
-		OutComparison = ELobbyComparisonOp::LessThanEquals;
-	}
-	else if (FCString::Stricmp(InStr, TEXT("Near")) == 0)
-	{
-		OutComparison = ELobbyComparisonOp::Near;
-	}
-	else if (FCString::Stricmp(InStr, TEXT("In")) == 0)
-	{
-		OutComparison = ELobbyComparisonOp::In;
-	}
-	else
-	{
-		checkNoEntry();
-		OutComparison = ELobbyComparisonOp::In;
-	}
-}
-
 const TCHAR* LexToString(EUILobbyJoinRequestedSource UILobbyJoinRequestedSource)
 {
 	switch (UILobbyJoinRequestedSource)
diff --git a/Engine/Plugins/Online/Experimental/OnlineServices/Source/OnlineServicesInterface/Private/Online/SchemaTypes.cpp b/Engine/Plugins/Online/Experimental/OnlineServices/Source/OnlineServicesInterface/Private/Online/SchemaTypes.cpp
--- a/Engine/Plugins/Online/Experimental/OnlineServices/Source/OnlineServicesInterface/Private/Online/SchemaTypes.cpp
+++ b/Engine/Plugins/Online/Experimental/OnlineServices/Source/OnlineServicesInterface/Private/Online/SchemaTypes.cpp
@@ -4,6 +4,91 @@
 
 namespace UE::Online {
 
+const TCHAR* LexToString(ESchemaAttributeComparisonOp Comparison)
+{
+	switch (Comparison)
+	{
+	case ESchemaAttributeComparisonOp::Equals:			return TEXT("Equals");
+	case ESchemaAttributeComparisonOp::NotEquals:			return TEXT("NotEquals");
+	case ESchemaAttributeComparisonOp::GreaterThan:		return TEXT("GreaterThan");
+	case ESchemaAttributeComparisonOp::GreaterThanEquals:	return TEXT("GreaterThanEquals");
+	case ESchemaAttributeComparisonOp::LessThan:			return TEXT("LessThan");
+	case ESchemaAttributeComparisonOp::LessThanEquals:	return TEXT("LessThanEquals");
+	case ESchemaAttributeComparisonOp::Near:				return TEXT("Near");
+	case ESchemaAttributeComparisonOp::In:				return TEXT("In");
+	default:									checkNoEntry(); // Intentional fallthrough
+	case ESchemaAttributeComparisonOp::NotIn:				return TEXT("NotIn");
+	}
+}
+
+void LexFromString(ESchemaAttributeComparisonOp& OutComparison, const TCHAR* InStr)
+{
+	if (FCString::Stricmp(InStr, TEXT("Equals")) == 0)
+	{
+		OutComparison = ESchemaAttributeComparisonOp::Equals;
+	}
+	else if (FCString::Stricmp(InStr, TEXT("NotEquals")) == 0)
+	{
+		OutComparison = ESchemaAttributeComparisonOp::NotEquals;
+	}
+	else if (FCString::Stricmp(InStr, TEXT("GreaterThan")) == 0)
+	{
+		OutComparison = ESchemaAttributeComparisonOp::GreaterThan;
+	}
+	else if (FCString::Stricmp(InStr, TEXT("GreaterThanEquals")) == 0)
+	{
+		OutComparison = ESchemaAttributeComparisonOp::GreaterThanEquals;
+	}
+	else if (FCString::Stricmp(InStr, TEXT("LessThan")) == 0)
+	{
+		OutComparison = ESchemaAttributeComparisonOp::LessThan;
+	}
+	else if (FCString::Stricmp(InStr, TEXT("LessThanEquals")) == 0)
+	{
+		OutComparison = ESchemaAttributeComparisonOp::LessThanEquals;
+	}
+	else if (FCString::Stricmp(InStr, TEXT("Near")) == 0)
+	{
+		OutComparison = ESchemaAttributeComparisonOp::Near;
+	}
+	else if (FCString::Stricmp(InStr, TEXT("In")) == 0)
+	{
+		OutComparison = ESchemaAttributeComparisonOp::In;
+	}
+	else
+	{
+		checkNoEntry();
+		OutComparison = ESchemaAttributeComparisonOp::In;
+	}
+}
+
+const TCHAR* LexToString(ESchemaAttributeVisibility Visibility)
+{
+	switch (Visibility)
+	{
+	case ESchemaAttributeVisibility::Public:		return TEXT("Public");
+	default:									checkNoEntry(); // Intentional fallthrough
+	case ESchemaAttributeVisibility::Private:	return TEXT("Private");
+	}
+}
+
+void LexFromString(ESchemaAttributeVisibility& OutVisibility, const TCHAR* InStr)
+{
+	if (FCString::Stricmp(InStr, TEXT("Public")) == 0)
+	{
+		OutVisibility = ESchemaAttributeVisibility::Public;
+	}
+	else if (FCString::Stricmp(InStr, TEXT("Private")) == 0)
+	{
+		OutVisibility = ESchemaAttributeVisibility::Private;
+	}
+	else
+	{
+		checkNoEntry();
+		OutVisibility = ESchemaAttributeVisibility::Private;
+	}
+}
+
 const TCHAR* LexToString(ESchemaAttributeFlags SchemaAttributeFlags)
 {
 	switch (SchemaAttributeFlags)
diff --git a/Engine/Plugins/Online/Experimental/OnlineServices/Source/OnlineServicesInterface/Private/Online/Sessions.cpp b/Engine/Plugins/Online/Experimental/OnlineServices/Source/OnlineServicesInterface/Private/Online/Sessions.cpp
--- a/Engine/Plugins/Online/Experimental/OnlineServices/Source/OnlineServicesInterface/Private/Online/Sessions.cpp
+++ b/Engine/Plugins/Online/Experimental/OnlineServices/Source/OnlineServicesInterface/Private/Online/Sessions.cpp
@@ -4,95 +4,6 @@
 
 namespace UE::Online {
 
-const TCHAR* LexToString(ESessionsComparisonOp Comparison)
-{
-	switch (Comparison)
-	{
-	case ESessionsComparisonOp::Equals:				return TEXT("Equals");
-	case ESessionsComparisonOp::NotEquals:			return TEXT("NotEquals");
-	case ESessionsComparisonOp::GreaterThan:		return TEXT("GreaterThan");
-	case ESessionsComparisonOp::GreaterThanEquals:	return TEXT("GreaterThanEquals");
-	case ESessionsComparisonOp::LessThan:			return TEXT("LessThan");
-	case ESessionsComparisonOp::LessThanEquals:		return TEXT("LessThanEquals");
-	default:										checkNoEntry(); // Intentional fallthrough
-	case ESessionsComparisonOp::Near:				return TEXT("Near");
-	}
-}
-
-void LexFromString(ESessionsComparisonOp& OutComparison, const TCHAR* InStr)
-{
-	if (FCString::Stricmp(InStr, TEXT("Equals")) == 0)
-	{
-		OutComparison = ESessionsComparisonOp::Equals;
-	}
-	else if (FCString::Stricmp(InStr, TEXT("NotEquals")) == 0)
-	{
-		OutComparison = ESessionsComparisonOp::NotEquals;
-	}
-	else if (FCString::Stricmp(InStr, TEXT("GreaterThan")) == 0)
-	{
-		OutComparison = ESessionsComparisonOp::GreaterThan;
-	}
-	else if (FCString::Stricmp(InStr, TEXT("GreaterThanEquals")) == 0)
-	{
-		OutComparison = ESessionsComparisonOp::GreaterThanEquals;
-	}
-	else if (FCString::Stricmp(InStr, TEXT("LessThan")) == 0)
-	{
-		OutComparison = ESessionsComparisonOp::LessThan;
-	}
-	else if (FCString::Stricmp(InStr, TEXT("LessThanEquals")) == 0)
-	{
-		OutComparison = ESessionsComparisonOp::LessThanEquals;
-	}
-	else if (FCString::Stricmp(InStr, TEXT("Near")) == 0)
-	{
-		OutComparison = ESessionsComparisonOp::Near;
-	}
-	else
-	{
-		checkNoEntry();
-		OutComparison = ESessionsComparisonOp::Near;
-	}
-}
-
-const TCHAR* LexToString(ECustomSessionSettingVisibility Value)
-{
-	switch (Value)
-	{
-	case ECustomSessionSettingVisibility::ViaPingOnly:				return TEXT("ViaPingOnly");
-	case ECustomSessionSettingVisibility::ViaOnlineService:			return TEXT("ViaOnlineService");
-	case ECustomSessionSettingVisibility::ViaOnlineServiceAndPing:	return TEXT("ViaOnlineServiceAndPing");
-	default:														checkNoEntry(); // Intentional fallthrough
-	case ECustomSessionSettingVisibility::DontAdvertise:			return TEXT("DontAdvertise");
-	}
-}
-
-void LexFromString(ECustomSessionSettingVisibility& Value, const TCHAR* InStr)
-{
-	if (FCString::Stricmp(InStr, TEXT("DontAdvertise")) == 0)
-	{
-		Value = ECustomSessionSettingVisibility::DontAdvertise;
-	}
-	else if (FCString::Stricmp(InStr, TEXT("ViaPingOnly")) == 0)
-	{
-		Value = ECustomSessionSettingVisibility::ViaPingOnly;
-	}
-	else if (FCString::Stricmp(InStr, TEXT("ViaOnlineService")) == 0)
-	{
-		Value = ECustomSessionSettingVisibility::ViaOnlineService;
-	}
-	else if (FCString::Stricmp(InStr, TEXT("ViaOnlineServiceAndPing")) == 0)
-	{
-		Value = ECustomSessionSettingVisibility::ViaOnlineServiceAndPing;
-	}
-	else
-	{
-		checkNoEntry();
-		Value = ECustomSessionSettingVisibility::DontAdvertise;
-	}
-}
-
 const TCHAR* LexToString(ESessionState Value)
 {
 	switch (Value)
diff --git a/Engine/Plugins/Online/Experimental/OnlineServices/Source/OnlineServicesInterface/Public/Online/Lobbies.h b/Engine/Plugins/Online/Experimental/OnlineServices/Source/OnlineServicesInterface/Public/Online/Lobbies.h
--- a/Engine/Plugins/Online/Experimental/OnlineServices/Source/OnlineServicesInterface/Public/Online/Lobbies.h
+++ b/Engine/Plugins/Online/Experimental/OnlineServices/Source/OnlineServicesInterface/Public/Online/Lobbies.h
@@ -10,9 +10,12 @@
 
 namespace UE::Online {
 
+// TODO: Aliases will be removed once full Schema support is added to Lobbies
 using FLobbySchemaId = FSchemaId;
 using FLobbyAttributeId = FSchemaAttributeId;
 using FLobbyVariant = FSchemaVariant;
+using ELobbyComparisonOp = ESchemaAttributeComparisonOp;
+using ELobbyAttributeVisibility = ESchemaAttributeVisibility;
 
 enum class ELobbyJoinPolicy : uint8
 {
@@ -51,14 +54,6 @@ enum class ELobbyMemberLeaveReason
 ONLINESERVICESINTERFACE_API const TCHAR* LexToString(ELobbyMemberLeaveReason LeaveReason);
 ONLINESERVICESINTERFACE_API void LexFromString(ELobbyMemberLeaveReason& OutLeaveReason, const TCHAR* InStr);
 
-enum class ELobbyAttributeVisibility
-{
-	Public,
-	Private,
-};
-ONLINESERVICESINTERFACE_API const TCHAR* LexToString(ELobbyAttributeVisibility Visibility);
-ONLINESERVICESINTERFACE_API void LexFromString(ELobbyAttributeVisibility& OutVisibility, const TCHAR* InStr);
-
 struct FLobbyMember
 {
 	FOnlineAccountIdHandle AccountId;
@@ -88,23 +83,6 @@ struct FJoinLobbyLocalUserData
 	TMap<FLobbyAttributeId, FLobbyVariant> Attributes;
 };
 
-enum class ELobbyComparisonOp : uint8
-{
-
-	Equals,
-	NotEquals,
-	GreaterThan,
-	GreaterThanEquals,
-	LessThan,
-	LessThanEquals,
-	Near,
-	In,
-	NotIn
-};
-
-ONLINESERVICESINTERFACE_API const TCHAR* LexToString(ELobbyComparisonOp EnumVal);
-ONLINESERVICESINTERFACE_API void LexFromString(ELobbyComparisonOp& OutComparison, const TCHAR* InStr);
-
 struct FFindLobbySearchFilter
 {
 	/** The name of the attribute to be compared. */
diff --git a/Engine/Plugins/Online/Experimental/OnlineServices/Source/OnlineServicesInterface/Public/Online/SchemaTypes.h b/Engine/Plugins/Online/Experimental/OnlineServices/Source/OnlineServicesInterface/Public/Online/SchemaTypes.h
--- a/Engine/Plugins/Online/Experimental/OnlineServices/Source/OnlineServicesInterface/Public/Online/SchemaTypes.h
+++ b/Engine/Plugins/Online/Experimental/OnlineServices/Source/OnlineServicesInterface/Public/Online/SchemaTypes.h
@@ -46,6 +46,31 @@ using FSchemaServiceAttributeId = FName;
   */
 using FSchemaServiceDescriptorId = FName;
 
+/** Flags to indicate comparison types between Schema Attributes. Used in searches */
+enum class ESchemaAttributeComparisonOp : uint8
+{
+	Equals,
+	NotEquals,
+	GreaterThan,
+	GreaterThanEquals,
+	LessThan,
+	LessThanEquals,
+	Near,
+	In,
+	NotIn
+};
+ONLINESERVICESINTERFACE_API const TCHAR* LexToString(ESchemaAttributeComparisonOp EnumVal);
+ONLINESERVICESINTERFACE_API void LexFromString(ESchemaAttributeComparisonOp& OutComparison, const TCHAR* InStr);
+
+// TODO: This enum fulfills the same function as the attribute flags, and will be removed once full Schema support is added to Sessions/Lobbies/Auth
+enum class ESchemaAttributeVisibility : uint8
+{
+	Public,
+	Private,
+};
+ONLINESERVICESINTERFACE_API const TCHAR* LexToString(ESchemaAttributeVisibility Visibility);
+ONLINESERVICESINTERFACE_API void LexFromString(ESchemaAttributeVisibility& OutVisibility, const TCHAR* InStr);
+
 /**
   * Flags to indicate required properties for a schema attribute.
   */
diff --git a/Engine/Plugins/Online/Experimental/OnlineServices/Source/OnlineServicesInterface/Public/Online/Sessions.h b/Engine/Plugins/Online/Experimental/OnlineServices/Source/OnlineServicesInterface/Public/Online/Sessions.h
--- a/Engine/Plugins/Online/Experimental/OnlineServices/Source/OnlineServicesInterface/Public/Online/Sessions.h
+++ b/Engine/Plugins/Online/Experimental/OnlineServices/Source/OnlineServicesInterface/Public/Online/Sessions.h
@@ -5,70 +5,36 @@
 #include "Online/OnlineAsyncOpHandle.h"
 #include "Online/CoreOnline.h"
 #include "Online/OnlineMeta.h"
+#include "Online/SchemaTypes.h"
 #include "Misc/TVariant.h"
 
 namespace UE::Online {
 
-enum class ESessionsComparisonOp : uint8
-{
-	Equals,
-	NotEquals,
-	GreaterThan,
-	GreaterThanEquals,
-	LessThan,
-	LessThanEquals,
-	Near
-};
-ONLINESERVICESINTERFACE_API const TCHAR* LexToString(ESessionsComparisonOp EnumVal);
-ONLINESERVICESINTERFACE_API void LexFromString(ESessionsComparisonOp& OutComparison, const TCHAR* InStr);
-
 struct FFindSessionsSearchFilter
 {
-	// TODO: Will change after SchemaVariant work
-
-	using FVariantType = TVariant<FString, int64, double, bool>;
-
 	/** Name of the custom setting to be used as filter */
-	FName Key;
+	FSchemaAttributeId Key;
 
 	/** The type of comparison to perform */
-	ESessionsComparisonOp ComparisonOp;
+	ESchemaAttributeComparisonOp ComparisonOp;
 
 	/** Value to use when comparing the filter */
-	FVariantType Value;
+	FSchemaVariant Value;
 };
 
-enum class ECustomSessionSettingVisibility : uint8
-{
-	/* Don't advertise via the online service or QoS data */
-	DontAdvertise,
-	/* Advertise via the server ping data only */
-	ViaPingOnly,
-	/* Advertise via the online service only */
-	ViaOnlineService,
-	/* Advertise via the online service and via the ping data */
-	ViaOnlineServiceAndPing
-};
-ONLINESERVICESINTERFACE_API const TCHAR* LexToString(ECustomSessionSettingVisibility Value);
-ONLINESERVICESINTERFACE_API void LexFromString(ECustomSessionSettingVisibility& Value, const TCHAR* InStr);
-
 struct FCustomSessionSetting
 {
-	// TODO: Will share this type with Lobbies
-
-	using FVariantType = TVariant<FString, int64, double, bool>;
-
 	/** Setting value */
-	FVariantType Data;
+	FSchemaVariant Data;
 
 	/** How is this session setting advertised with the backend or searches */
-	ECustomSessionSettingVisibility Visibility;
+	ESchemaAttributeVisibility Visibility;
 
 	/** Optional ID used in some platforms as the index instead of the setting name */
 	int32 ID;
 };
 
-using FCustomSessionSettingsMap = TMap<FName, FCustomSessionSetting>;
+using FCustomSessionSettingsMap = TMap<FSchemaAttributeId, FCustomSessionSetting>;
 
 /** A member is a player that is part of the session, and it stops being a member when they leave it */
 struct FSessionMember
@@ -81,7 +47,7 @@ using FSessionMembersMap = TMap<FOnlineAccountIdHandle, FSessionMember>;
 struct FSessionMemberUpdate
 {
 	FCustomSessionSettingsMap UpdatedMemberSettings;
-	TArray<FName> RemovedMemberSettings;
+	TArray<FSchemaAttributeId> RemovedMemberSettings;
 
 	FSessionMemberUpdate& operator+=(FSessionMemberUpdate&& UpdatedValue);
 };
@@ -124,7 +90,7 @@ using FRegisteredPlayersMap = TMap<FOnlineAccountIdHandle, FRegisteredPlayer>;
 struct FSessionSettingsUpdate
 {
 	/** Set with an updated value if the SchemaName field will be changed in the update operation */
-	TOptional<FName> SchemaName;
+	TOptional<FSchemaId> SchemaName;
 	/** Set with an updated value if the NumMaxPublicConnections field will be changed in the update operation */
 	TOptional<uint32> NumMaxPublicConnections;
 	/** Set with an updated value if the NumOpenPublicConnections field will be changed in the update operation */
@@ -153,7 +119,7 @@ struct FSessionSettingsUpdate
 	/** Updated values for custom settings to change in the update operation*/
 	FCustomSessionSettingsMap UpdatedCustomSettings;
 	/** Names of custom settings to be removed in the update operation*/
-	TArray<FName> RemovedCustomSettings;
+	TArray<FSchemaAttributeId> RemovedCustomSettings;
 
 	/** Updated values for session member info to change in the update operation*/
 	FSessionMemberUpdatesMap UpdatedSessionMembers;
@@ -171,7 +137,7 @@ struct FSessionSettingsUpdate
 struct ONLINESERVICESINTERFACE_API FSessionSettings
 {
 	/* The schema which will be applied to the session */
-	FName SchemaName;
+	FSchemaId SchemaName;
 
 	/* Maximum number of public slots for session members */
 	uint32 NumMaxPublicConnections = 0;
diff --git a/Engine/Plugins/Online/Experimental/OnlineServicesEOSGS/Source/Private/Online/SessionsEOSGS.cpp b/Engine/Plugins/Online/Experimental/OnlineServicesEOSGS/Source/Private/Online/SessionsEOSGS.cpp
--- a/Engine/Plugins/Online/Experimental/OnlineServicesEOSGS/Source/Private/Online/SessionsEOSGS.cpp
+++ b/Engine/Plugins/Online/Experimental/OnlineServicesEOSGS/Source/Private/Online/SessionsEOSGS.cpp
@@ -8,27 +8,12 @@
 #include "Online/OnlineServicesEOSGSTypes.h"
 #include "Online/NboSerializerEOSGSSvc.h"
 #include "Online/OnlineErrorEOSGS.h"
+#include "Online/SessionsEOSGSTypes.h"
 
 #include "eos_sessions.h"
-#include "eos_sessions_types.h"
 
 namespace UE::Online {
 
-/** Auxiliary methods */
-
-EOS_EOnlineSessionPermissionLevel GetEOSGSType(const ESessionJoinPolicy& Value)
-{
-	switch (Value)
-	{
-	case ESessionJoinPolicy::Public:		return EOS_EOnlineSessionPermissionLevel::EOS_OSPF_PublicAdvertised;
-	case ESessionJoinPolicy::FriendsOnly:	return EOS_EOnlineSessionPermissionLevel::EOS_OSPF_JoinViaPresence;
-	case ESessionJoinPolicy::InviteOnly:	return EOS_EOnlineSessionPermissionLevel::EOS_OSPF_InviteOnly;
-	}
-
-	checkNoEntry();
-	return EOS_EOnlineSessionPermissionLevel::EOS_OSPF_InviteOnly;
-}
-
 /** FOnlineSessionIdRegistryEOSGS */
 
 FOnlineSessionIdRegistryEOSGS::FOnlineSessionIdRegistryEOSGS()
@@ -70,12 +55,146 @@ FSessionEOSGS::FSessionEOSGS(const EOS_HSessionDetails& InSessionDetailsHandle)
 	{
 		CurrentState = ESessionState::Valid;
 		SessionId = FSessionsEOSGS::CreateSessionId(FString(SessionDetailsInfo->SessionId));
+		SessionSettings.NumOpenPublicConnections = SessionDetailsInfo->NumOpenPublicConnections;
+		// We could retrieve the Host Address here if we set it during creation or update
+
+		// TODO: evaluate effect of bInvitesAllowed and bAllowJoinInProgress to know whether we can just read the values here or if we need the custom parameter
+		SessionSettings.bAllowSanctionedPlayers = !SessionDetailsInfo->Settings->bSanctionsEnabled;
+		SessionSettings.NumMaxPublicConnections = SessionDetailsInfo->Settings->NumPublicConnections;
+		SessionSettings.JoinPolicy = FromServiceType(SessionDetailsInfo->Settings->PermissionLevel);
+		SessionSettings.CustomSettings.Emplace(EOSGS_BUCKET_ID, FCustomSessionSetting{ FSchemaVariant(FString(UTF8_TO_TCHAR(SessionDetailsInfo->Settings->BucketId))), ESchemaAttributeVisibility::Public });
+
+		// We retrieve all the session attributes
+		EOS_SessionDetails_GetSessionAttributeCountOptions GetAttributeCountOptions = {};
+		GetAttributeCountOptions.ApiVersion = EOS_SESSIONDETAILS_GETSESSIONATTRIBUTECOUNT_API_LATEST;
+		static_assert(EOS_SESSIONDETAILS_GETSESSIONATTRIBUTECOUNT_API_LATEST == 1, "EOS_SessionDetails_GetSessionAttributeCountOptions updated, check new fields");
+
+		uint32_t AttributeCount = EOS_SessionDetails_GetSessionAttributeCount(InSessionDetailsHandle, &GetAttributeCountOptions);
+		for (uint32_t Index = 0; Index < AttributeCount; ++Index)
+		{
+			EOS_SessionDetails_CopySessionAttributeByIndexOptions CopyAttributeByIndexOptions = {};
+			CopyAttributeByIndexOptions.ApiVersion = EOS_SESSIONDETAILS_COPYSESSIONATTRIBUTEBYINDEX_API_LATEST;
+			static_assert(EOS_SESSIONDETAILS_COPYSESSIONATTRIBUTEBYINDEX_API_LATEST == 1, "EOS_SessionDetails_CopySessionAttributeByIndexOptions updated, check new fields");
+
+			CopyAttributeByIndexOptions.AttrIndex = Index;
+
+			EOS_SessionDetails_Attribute* Attribute = nullptr;
+			EOS_EResult CopyAttributeByIndexResult = EOS_SessionDetails_CopySessionAttributeByIndex(InSessionDetailsHandle, &CopyAttributeByIndexOptions, &Attribute);
+			if (CopyAttributeByIndexResult == EOS_EResult::EOS_Success)
+			{
+				// We parse a single attribute
+				FString Key(Attribute->Data->Key);
+
+				// If the Key contains the ':' character, it will be a user related entry, either a Session Member setting, or Registered Player information
+				if (Key.Contains(TEXT(":")))
+				{
+					TArray<FString> KeyComponents;
+					Key.ParseIntoArray(KeyComponents, TEXT(":"));
+
+					// If the first element is the Registered Players key, we will process the entry as Registered Player information
+					if (KeyComponents[0] == EOSGS_REGISTERED_PLAYERS.ToString())
+					{
+						// We retrieve the registered player id
+						const FString& PlayerIdStr = KeyComponents[1];
+						const EOS_ProductUserId ProductUserId = EOS_ProductUserId_FromString(TCHAR_TO_UTF8(*PlayerIdStr));
+						const FOnlineAccountIdHandle& IdHandle = FindAccountId(ProductUserId);
+
+						// This constructor should only be used by BuildSessionFromDetailsHandle, after all user ids in the session have been resolved
+						check(IdHandle.IsValid());
+
+						FRegisteredPlayer& RegisteredPlayer = SessionSettings.RegisteredPlayers.FindOrAdd(IdHandle);
+						
+						FSessionAttributeConverter<ESessionAttributeConversionType::FromService> CustomSettingConverter(*Attribute->Data);
+						const TPair<FSchemaAttributeId, FSchemaVariant>& CustomSettingData = CustomSettingConverter.GetAttributeData();
+						bool bParameterValue = CustomSettingData.Value.GetBoolean();
+
+						// And set the value for the appropriate parameter
+						const FString& ParameterName = KeyComponents[2];
+
+						if (ParameterName == EOSGS_REGISTERED_PLAYER_IS_IN_SESSION.ToString())
+						{
+							RegisteredPlayer.bIsInSession = bParameterValue;
+						}
+						else if (ParameterName == EOSGS_REGISTERED_PLAYER_HAS_RESERVED_SLOT.ToString())
+						{
+							RegisteredPlayer.bHasReservedSlot = bParameterValue;
+						}
+					}
+					else // If not, it will be parsed as a Session Member setting
+					{
+						// We retrieve the member id
+						FString PlayerIdStr = KeyComponents[0];
+						EOS_ProductUserId ProductUserId = EOS_ProductUserId_FromString(TCHAR_TO_UTF8(*PlayerIdStr));
+						FOnlineAccountIdHandle IdHandle = FindAccountId(ProductUserId);
+
+						FSessionMember& SessionMember = SessionSettings.SessionMembers.FindOrAdd(IdHandle);
+
+						// And add the corresponding custom setting
+						FSchemaAttributeId AttributeId(KeyComponents[1]);
+
+						FSessionAttributeConverter<ESessionAttributeConversionType::FromService> CustomSettingConverter(*Attribute->Data);
+						TPair<FSchemaAttributeId, FSchemaVariant> CustomSettingData = CustomSettingConverter.GetAttributeData();
+
+						FCustomSessionSetting CustomSessionSetting;
+						CustomSessionSetting.Visibility = FromServiceType(Attribute->AdvertisementType);
+						CustomSessionSetting.Data = CustomSettingData.Value;
+
+						SessionMember.MemberSettings.Emplace(AttributeId, CustomSessionSetting);
+					}
+				}
+				else 
+				{
+					FSessionAttributeConverter<ESessionAttributeConversionType::FromService> CustomSettingConverter(*Attribute->Data);
+					TPair<FSchemaAttributeId, FSchemaVariant> CustomSettingData = CustomSettingConverter.GetAttributeData();
+					
+					// Most Session Settings values get parsed in the same way as Custom Session Settings, so we will attempt to retrieve them
+					if (Key == EOSGS_ALLOW_NEW_MEMBERS.ToString())
+					{
+						SessionSettings.bAllowNewMembers = CustomSettingData.Value.GetBoolean();
+					}
+					else if (Key == EOSGS_ALLOW_UNREGISTERED_PLAYERS.ToString())
+					{
+						SessionSettings.bAllowUnregisteredPlayers = CustomSettingData.Value.GetBoolean();
+					}
+					else if (Key == EOSGS_ANTI_CHEAT_PROTECTED.ToString())
+					{
+						SessionSettings.bAntiCheatProtected = CustomSettingData.Value.GetBoolean();
+					}
+					else if (Key == EOSGS_IS_DEDICATED_SERVER_SESSION.ToString())
+					{
+						SessionSettings.bIsDedicatedServerSession = CustomSettingData.Value.GetBoolean();
+					}
+					else if (Key == EOSGS_PRESENCE_ENABLED.ToString())
+					{
+						SessionSettings.bPresenceEnabled = CustomSettingData.Value.GetBoolean();
+					}
+					else if (Key == EOSGS_SCHEMA_NAME.ToString())
+					{
+						SessionSettings.SchemaName = FSchemaId(CustomSettingData.Value.GetString());
+					}
+					else if (Key == EOSGS_SESSION_ID_OVERRIDE.ToString())
+					{
+						SessionSettings.SessionIdOverride = CustomSettingData.Value.GetString();
+					}
+					else // The rest are parsed as a Custom Session Setting
+					{
+						FCustomSessionSetting CustomSessionSetting;
+						CustomSessionSetting.Visibility = FromServiceType(Attribute->AdvertisementType);
+						CustomSessionSetting.Data = CustomSettingData.Value;
 
-		// TODO: Load all session settings after SchemaVariant work
+						SessionSettings.CustomSettings.Emplace(Key, CustomSessionSetting);
+					}
+				}
+			}
+			else
+			{
+				UE_LOG(LogTemp, Warning, TEXT("[FSessionEOSGS] EOS_SessionDetails_CopySessionAttributeByIndex failed with result [%s]"), *LexToString(CopyInfoResult));
+			}
+		}
 	}
 	else
 	{
-		UE_LOG(LogTemp, Error, TEXT("[FSessionEOSGS] EOS_SessionDetails_CopyInfo failed with result [%s]"), *LexToString(CopyInfoResult));
+		UE_LOG(LogTemp, Warning, TEXT("[FSessionEOSGS] EOS_SessionDetails_CopyInfo failed with result [%s]"), *LexToString(CopyInfoResult));
 
 		CurrentState = ESessionState::Invalid;
 	}
@@ -287,8 +406,8 @@ TOnlineAsyncOpHandle<FCreateSession> FSessionsEOSGS::CreateSession(FCreateSessio
 		CreateSessionModificationOptions.bPresenceEnabled = OpParams.SessionSettings.bPresenceEnabled;
 		CreateSessionModificationOptions.bSanctionsEnabled = !OpParams.SessionSettings.bAllowSanctionedPlayers;
 
-		const FCustomSessionSetting* BucketIdSetting = OpParams.SessionSettings.CustomSettings.Find(EOS_SESSIONS_BUCKET_ID);
-		FTCHARToUTF8 BucketIdUtf8(BucketIdSetting != nullptr ? BucketIdSetting->Data.Get<FString>() : FString());
+		const FCustomSessionSetting* BucketIdSetting = OpParams.SessionSettings.CustomSettings.Find(EOSGS_BUCKET_ID);
+		FTCHARToUTF8 BucketIdUtf8(BucketIdSetting != nullptr ? BucketIdSetting->Data.GetString() : FString());
 
 		if (BucketIdUtf8.Length())
 		{
@@ -368,12 +487,12 @@ void FSessionsEOSGS::SetPermissionLevel(EOS_HSessionModification& SessionModific
 	Options.ApiVersion = EOS_SESSIONMODIFICATION_SETPERMISSIONLEVEL_API_LATEST;
 	static_assert(EOS_SESSIONMODIFICATION_SETPERMISSIONLEVEL_API_LATEST == 1, "EOS_SessionModification_SetPermissionLevelOptions updated, check new fields");
 
-	Options.PermissionLevel = GetEOSGSType(NewJoinPolicy);
+	Options.PermissionLevel = ToServiceType(NewJoinPolicy);
 
 	EOS_EResult ResultCode = EOS_SessionModification_SetPermissionLevel(SessionModificationHandle, &Options);
 	if (ResultCode != EOS_EResult::EOS_Success)
 	{
-		UE_LOG(LogTemp, Warning, TEXT("EOS_SessionModification_SetPermissionLevel failed with result [%s]"), *LexToString(ResultCode));
+		UE_LOG(LogTemp, Warning, TEXT("[FSessionsEOSGS::SetPermissionLevel] EOS_SessionModification_SetPermissionLevel failed with result [%s]"), *LexToString(ResultCode));
 	}
 }
 
@@ -389,7 +508,7 @@ void FSessionsEOSGS::SetBucketId(EOS_HSessionModification& SessionModificationHa
 	EOS_EResult ResultCode = EOS_SessionModification_SetBucketId(SessionModificationHandle, &Options);
 	if (ResultCode != EOS_EResult::EOS_Success)
 	{
-		UE_LOG(LogTemp, Warning, TEXT("EOS_SessionModification_SetBucketId failed with result [%s]"), *LexToString(ResultCode));
+		UE_LOG(LogTemp, Warning, TEXT("[FSessionsEOSGS::SetBucketId] EOS_SessionModification_SetBucketId failed with result [%s]"), *LexToString(ResultCode));
 	}
 }
 
@@ -404,7 +523,41 @@ void FSessionsEOSGS::SetMaxPlayers(EOS_HSessionModification& SessionModification
 	EOS_EResult ResultCode = EOS_SessionModification_SetMaxPlayers(SessionModificationHandle, &Options);
 	if (ResultCode != EOS_EResult::EOS_Success)
 	{
-		UE_LOG(LogTemp, Warning, TEXT("EOS_SessionModification_SetMaxPlayers failed with result [%s]"), *LexToString(ResultCode));
+		UE_LOG(LogTemp, Warning, TEXT("[FSessionsEOSGS::SetMaxPlayers] EOS_SessionModification_SetMaxPlayers failed with result [%s]"), *LexToString(ResultCode));
+	}
+}
+
+void FSessionsEOSGS::AddAttribute(EOS_HSessionModification& SessionModificationHandle, const FSchemaAttributeId& Key, const FCustomSessionSetting& Value)
+{
+	EOS_SessionModification_AddAttributeOptions Options = { };
+	Options.ApiVersion = EOS_SESSIONMODIFICATION_ADDATTRIBUTE_API_LATEST;
+	static_assert(EOS_SESSIONMODIFICATION_ADDATTRIBUTE_API_LATEST == 1, "EOS_SessionModification_AddAttributeOptions updated, check new fields");
+
+	Options.AdvertisementType = ToServiceType(Value.Visibility);
+
+	FSessionAttributeConverter<ESessionAttributeConversionType::ToService> SessionAttribute(Key, Value.Data);
+	Options.SessionAttribute = &SessionAttribute.GetAttributeData();
+
+	EOS_EResult ResultCode = EOS_SessionModification_AddAttribute(SessionModificationHandle, &Options);
+	if (ResultCode != EOS_EResult::EOS_Success)
+	{
+		UE_LOG(LogTemp, Warning, TEXT("[FSessionsEOSGS::AddAttribute] EOS_SessionModification_AddAttribute failed with result [%s]"), *LexToString(ResultCode));
+	}
+}
+
+void FSessionsEOSGS::RemoveAttribute(EOS_HSessionModification& SessionModificationHandle, const FSchemaAttributeId& Key)
+{
+	EOS_SessionModification_RemoveAttributeOptions Options = { };
+	Options.ApiVersion = EOS_SESSIONMODIFICATION_REMOVEATTRIBUTE_API_LATEST;
+	static_assert(EOS_SESSIONMODIFICATION_REMOVEATTRIBUTE_API_LATEST == 1, "EOS_SessionModification_RemoveAttributeOptions updated, check new fields");
+
+	FTCHARToUTF8 KeyUtf8(Key.ToString());
+	Options.Key = KeyUtf8.Get();
+
+	EOS_EResult ResultCode = EOS_SessionModification_RemoveAttribute(SessionModificationHandle, &Options);
+	if (ResultCode != EOS_EResult::EOS_Success)
+	{
+		UE_LOG(LogTemp, Warning, TEXT("[FSessionsEOSGS::RemoveAttribute] EOS_SessionModification_RemoveAttribute failed with result [%s]"), *LexToString(ResultCode));
 	}
 }
 
@@ -415,46 +568,57 @@ void FSessionsEOSGS::WriteCreateSessionModificationHandle(EOS_HSessionModificati
 	// TODO: We either call Start or End session, EOS_SessionModification_SetJoinInProgressAllowed and EOS_SessionModification_SetInvitesAllowed depending on this value changing, or a custom setting.
 	// Investigation pending to see if there is overlay changes when the session is not joinable (started + JIP false, or InvitesAllowed false)
 
-	// TODO: Add non-custom settings as attributes on the API call to transmit that information in session searches. Example:
-
-// 	FCustomSessionSetting NewSetting;
-// 	NewSetting.Data.Set<bool>(SessionSettings.bAllowNewMembers);
-// 	NewSetting.Visibility = ECustomSessionSettingVisibility::ViaOnlineService;
-// 
-// 	AddAttribute(SessionModificationHandle, EOS_SESSIONS_ALLOW_NEW_MEMBERS, NewSetting);
+	// We won't copy bIsLANSession since it' irrelevant for EOS Sessions
+	AddAttribute(SessionModificationHandle, EOSGS_ALLOW_NEW_MEMBERS, { FSchemaVariant(SessionSettings.bAllowNewMembers), ESchemaAttributeVisibility::Public });
+	AddAttribute(SessionModificationHandle, EOSGS_ALLOW_UNREGISTERED_PLAYERS, { FSchemaVariant(SessionSettings.bAllowUnregisteredPlayers), ESchemaAttributeVisibility::Public });
+	AddAttribute(SessionModificationHandle, EOSGS_ANTI_CHEAT_PROTECTED, { FSchemaVariant(SessionSettings.bAntiCheatProtected), ESchemaAttributeVisibility::Public });
+	AddAttribute(SessionModificationHandle, EOSGS_IS_DEDICATED_SERVER_SESSION, { FSchemaVariant(SessionSettings.bIsDedicatedServerSession), ESchemaAttributeVisibility::Public });
+	AddAttribute(SessionModificationHandle, EOSGS_PRESENCE_ENABLED, { FSchemaVariant(SessionSettings.bPresenceEnabled), ESchemaAttributeVisibility::Public });
 
 	SetPermissionLevel(SessionModificationHandle, SessionSettings.JoinPolicy);
 
 	SetMaxPlayers(SessionModificationHandle, SessionSettings.NumMaxPrivateConnections + SessionSettings.NumMaxPublicConnections);
 
+	AddAttribute(SessionModificationHandle, EOSGS_SCHEMA_NAME, { FSchemaVariant(SessionSettings.SchemaName.ToString()), ESchemaAttributeVisibility::Public });
+	AddAttribute(SessionModificationHandle, EOSGS_SESSION_ID_OVERRIDE, { FSchemaVariant(SessionSettings.SessionIdOverride), ESchemaAttributeVisibility::Public });
+
 	// Custom Settings
 
-// 	for (FCustomSessionSettingsMap::TConstIterator It(SessionSettings.CustomSettings); It; ++It)
-// 	{
-// 		// TODO: Add the attribute to the SessionModification. Pending SchemaVariant work.
-// 	}
+ 	for (const TPair<FSchemaAttributeId, FCustomSessionSetting>& Entry : SessionSettings.CustomSettings)
+ 	{
+		AddAttribute(SessionModificationHandle, Entry.Key, Entry.Value);
+	}
 
 	// BucketId has its own set method on the API
-	const FCustomSessionSetting* NewBucketIdSetting = SessionSettings.CustomSettings.Find(EOS_SESSIONS_BUCKET_ID);
+	const FCustomSessionSetting* NewBucketIdSetting = SessionSettings.CustomSettings.Find(EOSGS_BUCKET_ID);
 	if (NewBucketIdSetting)
 	{
-		SetBucketId(SessionModificationHandle, NewBucketIdSetting->Data.Get<FString>());
+		SetBucketId(SessionModificationHandle, NewBucketIdSetting->Data.GetString());
 	}
 
 	// Session Members
 
-// 	for (FSessionMembersMap::TConstIterator MemberIt(SessionSettings.SessionMembers); MemberIt; ++MemberIt)
-// 	{
-// 		const FSessionMember& SessionMember = MemberIt.Value();
-// 		
-// 		for (FCustomSessionSettingsMap::TConstIterator SettingIt(SessionMember.MemberSettings); SettingIt; ++SettingIt)
-// 		{
-// 			const FName& Key = FName(LexToString(GetProductUserIdChecked(MemberIt.Key())) + TEXT(":") + SettingIt.Key().ToString());
-// 			const FCustomSessionSetting& Value = SettingIt.Value();
-// 
-// 			// TODO: Add the attribute to the SessionModification. Pending SchemaVariant work.
-// 		}		
-// 	}
+	for (const TPair<FOnlineAccountIdHandle, FSessionMember>& SessionMemberEntry : SessionSettings.SessionMembers)
+	{
+		const FSessionMember& SessionMember = SessionMemberEntry.Value;
+		
+		for (const TPair<FSchemaAttributeId, FCustomSessionSetting>& CustomSettingEntry : SessionMember.MemberSettings)
+		{
+			const FName& Key = FName(LexToString(GetProductUserIdChecked(SessionMemberEntry.Key)) + TEXT(":") + CustomSettingEntry.Key.ToString());
+			AddAttribute(SessionModificationHandle, Key, CustomSettingEntry.Value);
+		}		
+	}
+
+	// TODO: Should always be empty at session creation time so might not be needed
+	// Registered Players
+	for (const TPair<FOnlineAccountIdHandle, FRegisteredPlayer>& Entry : SessionSettings.RegisteredPlayers)
+	{
+		const FName& HasReservedSlotKey = FName(EOSGS_REGISTERED_PLAYERS.ToString() + TEXT(":") + LexToString(GetProductUserIdChecked(Entry.Key)) + TEXT(":") + EOSGS_REGISTERED_PLAYER_HAS_RESERVED_SLOT.ToString());
+		AddAttribute(SessionModificationHandle, HasReservedSlotKey, { FSchemaVariant(Entry.Value.bHasReservedSlot), ESchemaAttributeVisibility::Public });
+
+		const FName& IsInSessionKey = FName(EOSGS_REGISTERED_PLAYERS.ToString() + TEXT(":") + LexToString(GetProductUserIdChecked(Entry.Key)) + TEXT(":") + EOSGS_REGISTERED_PLAYER_IS_IN_SESSION.ToString());
+		AddAttribute(SessionModificationHandle, IsInSessionKey, { FSchemaVariant(Entry.Value.bIsInSession), ESchemaAttributeVisibility::Public });
+	}
 }
 
 void FSessionsEOSGS::WriteUpdateSessionModificationHandle(EOS_HSessionModification& SessionModificationHandle, const FName& SessionName, const FSessionSettingsUpdate& NewSettings)
@@ -464,14 +628,43 @@ void FSessionsEOSGS::WriteUpdateSessionModificationHandle(EOS_HSessionModificati
 	// TODO: We either call Start or End session, EOS_SessionModification_SetJoinInProgressAllowed and EOS_SessionModification_SetInvitesAllowed depending on this value changing, or a custom setting.
 	// Investigation pending to see if there is overlay changes when the session is not joinable (started + JIP false, or InvitesAllowed false)
 
-	// TODO: Add non-custom settings as attributes on the API call to transmit that information in session searches. Example:
+	if (NewSettings.bAllowNewMembers.IsSet())
+	{
+		AddAttribute(SessionModificationHandle, EOSGS_ALLOW_NEW_MEMBERS, { FSchemaVariant(NewSettings.bAllowNewMembers.GetValue()), ESchemaAttributeVisibility::Public });
+	}
+
+	if (NewSettings.bAllowUnregisteredPlayers.IsSet())
+	{
+		AddAttribute(SessionModificationHandle, EOSGS_ALLOW_UNREGISTERED_PLAYERS, { FSchemaVariant(NewSettings.bAllowUnregisteredPlayers.GetValue()), ESchemaAttributeVisibility::Public });
+	}
+
+	if (NewSettings.bAntiCheatProtected.IsSet())
+	{
+		AddAttribute(SessionModificationHandle, EOSGS_ANTI_CHEAT_PROTECTED, { FSchemaVariant(NewSettings.bAntiCheatProtected.GetValue()), ESchemaAttributeVisibility::Public });
+	}
+
+	if (NewSettings.bIsDedicatedServerSession.IsSet())
+	{
+		AddAttribute(SessionModificationHandle, EOSGS_IS_DEDICATED_SERVER_SESSION, { FSchemaVariant(NewSettings.bIsDedicatedServerSession.GetValue()), ESchemaAttributeVisibility::Public });
+	}
+
+	if (NewSettings.bPresenceEnabled.IsSet())
+	{
+		AddAttribute(SessionModificationHandle, EOSGS_PRESENCE_ENABLED, { FSchemaVariant(NewSettings.bPresenceEnabled.GetValue()), ESchemaAttributeVisibility::Public });
+	}
+
+	if (NewSettings.SchemaName.IsSet())
+	{
+		AddAttribute(SessionModificationHandle, EOSGS_SCHEMA_NAME, { FSchemaVariant(NewSettings.SchemaName.GetValue().ToString()), ESchemaAttributeVisibility::Public });
+	}
+
+	if (NewSettings.SessionIdOverride.IsSet())
+	{
+		AddAttribute(SessionModificationHandle, EOSGS_SESSION_ID_OVERRIDE, { FSchemaVariant(NewSettings.SessionIdOverride.GetValue()), ESchemaAttributeVisibility::Public });
+	}
 
-// 	FCustomSessionSetting NewSetting;
-// 	NewSetting.Data.Set<bool>(SessionSettings.bAllowNewMembers);
-// 	NewSetting.Visibility = ECustomSessionSettingVisibility::ViaOnlineService;
-// 
-// 	AddAttribute(SessionModificationHandle, EOS_SESSIONS_ALLOW_NEW_MEMBERS, NewSetting);
 	const TSharedRef<FSession>& FoundSession = LocalSessionsByName.FindChecked(SessionName);
+	const FSessionSettings& SessionSettings = FoundSession->SessionSettings;
 
 	if (NewSettings.JoinPolicy.IsSet())
 	{
@@ -481,94 +674,84 @@ void FSessionsEOSGS::WriteUpdateSessionModificationHandle(EOS_HSessionModificati
 	if (NewSettings.NumMaxPrivateConnections.IsSet() || NewSettings.NumMaxPublicConnections.IsSet())
 	{
 		// We set MaxPlayers to the sum of any new values (if set) and old values (if not set)
-		uint32 DefaultNumMaxPrivateConnections = FoundSession->SessionSettings.NumMaxPrivateConnections;
-		uint32 DefaultNumMaxPublicConnections = FoundSession->SessionSettings.NumMaxPublicConnections;
+		uint32 DefaultNumMaxPrivateConnections = SessionSettings.NumMaxPrivateConnections;
+		uint32 DefaultNumMaxPublicConnections = SessionSettings.NumMaxPublicConnections;
 
 		SetMaxPlayers(SessionModificationHandle, NewSettings.NumMaxPrivateConnections.Get(DefaultNumMaxPrivateConnections) + NewSettings.NumMaxPublicConnections.Get(DefaultNumMaxPublicConnections));
 	}
 
 	// BucketId has its own set method on the API
-	const FCustomSessionSetting* NewBucketIdSetting = NewSettings.UpdatedCustomSettings.Find(EOS_SESSIONS_BUCKET_ID);
+	const FCustomSessionSetting* NewBucketIdSetting = NewSettings.UpdatedCustomSettings.Find(EOSGS_BUCKET_ID);
 	if (NewBucketIdSetting)
 	{
-		SetBucketId(SessionModificationHandle, NewBucketIdSetting->Data.Get<FString>());
+		SetBucketId(SessionModificationHandle, NewBucketIdSetting->Data.GetString());
 	}
 
-	// TODO: This next section will be reactivated as part of the CL to add support for SchemaVariant
-#if 0
 	// Custom Settings
 
-	for (const FName& Key : NewSettings.RemovedCustomSettings)
+	for (const FSchemaAttributeId& Key : NewSettings.RemovedCustomSettings)
 	{
-		SessionSettings.CustomSettings.Remove(Key);
-
-		// TODO: Remove the attribute to the SessionModification. Pending SchemaVariant work.
+		RemoveAttribute(SessionModificationHandle, Key);
 	}
 
-// 	for (FCustomSessionSettingsMap::TConstIterator It(NewSettings.UpdatedCustomSettings); It; ++It)
-// 	{
-// 		// TODO: Add the attribute to the SessionModification. Pending SchemaVariant work.
-// 	}
-
-	SessionSettings.CustomSettings.Append(NewSettings.UpdatedCustomSettings);
-
-	// registered players
+	for (const TPair<FSchemaAttributeId, FCustomSessionSetting>& Entry : NewSettings.UpdatedCustomSettings)
+	{
+		AddAttribute(SessionModificationHandle, Entry.Key, Entry.Value);
+	}
 
-	for (const FOnlineAccountIdHandle& Key : NewSettings.RemovedRegisteredPlayers)
+	// Session Members
+	
+	for (const FOnlineAccountIdHandle& SessionMemberId : NewSettings.RemovedSessionMembers)
 	{
-		SessionSettings.RegisteredPlayers.Remove(Key);
+		// To remove a session member, we'll remove all their attributes
+ 		if (const FSessionMember* SessionMember = SessionSettings.SessionMembers.Find(SessionMemberId))
+ 		{
+ 			TArray<FName> SettingsKeyArray;
+ 			SessionMember->MemberSettings.GenerateKeyArray(SettingsKeyArray);
+ 			for (const FSchemaAttributeId& CustomSettingEntryKey : SettingsKeyArray)
+ 			{
+ 				const FSchemaAttributeId& Key = FName(LexToString(GetProductUserIdChecked(SessionMemberId)) + TEXT(":") + CustomSettingEntryKey.ToString());
+				RemoveAttribute(SessionModificationHandle, Key);
+ 			}
+ 		}
 	}
 
-	SessionSettings.RegisteredPlayers.Append(NewSettings.UpdatedRegisteredPlayers);
+	for (const TPair<FOnlineAccountIdHandle, FSessionMemberUpdate>& SessionMemberEntry : NewSettings.UpdatedSessionMembers)
+	{
+		const FSessionMemberUpdate& SessionMemberUpdate = SessionMemberEntry.Value;
 
-	// Session Members
+		for (const FSchemaAttributeId& CustomSettingEntryKey : SessionMemberUpdate.RemovedMemberSettings)
+		{
+			const FSchemaAttributeId& Key = FSchemaAttributeId(LexToString(GetProductUserIdChecked(SessionMemberEntry.Key)) + TEXT(":") + CustomSettingEntryKey.ToString());
+			RemoveAttribute(SessionModificationHandle, Key);
+		}
 
-	for (const FOnlineAccountIdHandle& Key : NewSettings.RemovedSessionMembers)
-	{
-// 		// Before removing a session member, we'll remove all their attributes
-// 		if (FSessionMember* SessionMember = SessionSettings.SessionMembers.Find(Key))
-// 		{
-// 			TArray<FName> SettingsKeyArray;
-// 			SessionMember->MemberSettings.GenerateKeyArray(SettingsKeyArray);
-// 
-// 			for (const FName& SettingKey : SettingsKeyArray)
-// 			{
-// 				const FName& MemberKey = FName(LexToString(GetProductUserIdChecked(Key)) + TEXT(":") + SettingKey.ToString());
-// 				
-// 				// TODO: Remove the attribute to the SessionModification. Pending SchemaVariant work.
-// 			}
-// 		}
-
-		SessionSettings.SessionMembers.Remove(Key);
+		for (const TPair<FSchemaAttributeId, FCustomSessionSetting>& CustomSettingEntry : SessionMemberUpdate.UpdatedMemberSettings)
+		{
+			const FSchemaAttributeId& Key = FSchemaAttributeId(LexToString(GetProductUserIdChecked(SessionMemberEntry.Key)) + TEXT(":") + CustomSettingEntry.Key.ToString());
+			AddAttribute(SessionModificationHandle, Key, CustomSettingEntry.Value);
+		}
 	}
 
-	for (FSessionMemberUpdatesMap::TConstIterator MemberIt(NewSettings.UpdatedSessionMembers); MemberIt; ++MemberIt)
+	// Registered Players
+	
+	for (const FOnlineAccountIdHandle& RemovedEntry : NewSettings.RemovedRegisteredPlayers)
 	{
-		if (FSessionMember* SessionMember = SessionSettings.SessionMembers.Find(MemberIt.Key()))
-		{
-			const FSessionMemberUpdate& SessionMemberUpdate = MemberIt.Value();
+		const FSchemaAttributeId& HasReservedSlotKey = FSchemaAttributeId(EOSGS_REGISTERED_PLAYERS.ToString() + TEXT(":") + LexToString(GetProductUserIdChecked(RemovedEntry)) + TEXT(":") + EOSGS_REGISTERED_PLAYER_HAS_RESERVED_SLOT.ToString());
+		RemoveAttribute(SessionModificationHandle, HasReservedSlotKey);
 
-			for (const FName& Key : SessionMemberUpdate.RemovedMemberSettings)
-			{
-				SessionMember->MemberSettings.Remove(Key);
+		const FSchemaAttributeId& IsInSessionKey = FSchemaAttributeId(EOSGS_REGISTERED_PLAYERS.ToString() + TEXT(":") + LexToString(GetProductUserIdChecked(RemovedEntry)) + TEXT(":") + EOSGS_REGISTERED_PLAYER_IS_IN_SESSION.ToString());
+		RemoveAttribute(SessionModificationHandle, IsInSessionKey);
+	}
 
-// 				const FName& MemberKey = FName(LexToString(GetProductUserIdChecked(MemberIt.Key())) + TEXT(":") + Key.ToString());
-// 				
-// 				// TODO: Remove the attribute to the SessionModification. Pending SchemaVariant work.
-			}
-			
-// 			for (FCustomSessionSettingsMap::TConstIterator SettingIt(SessionMemberUpdate.UpdatedMemberSettings); SettingIt; ++SettingIt)
-// 			{
-// 				const FName& Key = FName(LexToString(GetProductUserIdChecked(MemberIt.Key())) + TEXT(":") + SettingIt.Key().ToString());
-// 				const FCustomSessionSetting& Value = SettingIt.Value();
-// 
-// 				// TODO: Add the attribute to the SessionModification. Pending SchemaVariant work.
-// 			}
-
-			SessionMember->MemberSettings.Append(SessionMemberUpdate.UpdatedMemberSettings);
-		}
+	for (const TPair<FOnlineAccountIdHandle, FRegisteredPlayer>& Entry : NewSettings.UpdatedRegisteredPlayers)
+	{
+		const FSchemaAttributeId& HasReservedSlotKey = FSchemaAttributeId(EOSGS_REGISTERED_PLAYERS.ToString() + TEXT(":") + LexToString(GetProductUserIdChecked(Entry.Key)) + TEXT(":") + EOSGS_REGISTERED_PLAYER_HAS_RESERVED_SLOT.ToString());
+		AddAttribute(SessionModificationHandle, HasReservedSlotKey, { FSchemaVariant(Entry.Value.bHasReservedSlot), ESchemaAttributeVisibility::Public });
+
+		const FSchemaAttributeId& IsInSessionKey = FSchemaAttributeId(EOSGS_REGISTERED_PLAYERS.ToString() + TEXT(":") + LexToString(GetProductUserIdChecked(Entry.Key)) + TEXT(":") + EOSGS_REGISTERED_PLAYER_IS_IN_SESSION.ToString());
+		AddAttribute(SessionModificationHandle, IsInSessionKey, { FSchemaVariant(Entry.Value.bIsInSession), ESchemaAttributeVisibility::Public });
 	}
-#endif
 }
 
 TOnlineAsyncOpHandle<FUpdateSession> FSessionsEOSGS::UpdateSession(FUpdateSession::Params&& Params)
@@ -766,8 +949,23 @@ void FSessionsEOSGS::SetSessionSearchMaxResults(FSessionSearchHandleEOSGS& Sessi
 
 void FSessionsEOSGS::SetSessionSearchParameters(FSessionSearchHandleEOSGS& SessionSearchHandle, TArray<FFindSessionsSearchFilter> Filters)
 {
-	// TODO: Pending SchemaVariant work
-	// Here we'll call EOS_SessionSearch_SetParameter
+	for (const FFindSessionsSearchFilter& Filter : Filters)
+	{
+		EOS_SessionSearch_SetParameterOptions Options = { };
+		Options.ApiVersion = EOS_SESSIONSEARCH_SETPARAMETER_API_LATEST;
+		static_assert(EOS_SESSIONSEARCH_SETPARAMETER_API_LATEST == 1, "EOS_SessionSearch_SetMaxResultsOptions updated, check new fields");
+
+		Options.ComparisonOp = ToServiceType(Filter.ComparisonOp);
+		
+		FSessionAttributeConverter<ESessionAttributeConversionType::ToService> Parameter(Filter.Key, Filter.Value);
+		Options.Parameter = &Parameter.GetAttributeData();
+
+		EOS_EResult ResultCode = EOS_SessionSearch_SetParameter(SessionSearchHandle.SearchHandle, &Options);
+		if (ResultCode != EOS_EResult::EOS_Success)
+		{
+			UE_LOG(LogTemp, Warning, TEXT("[FSessionsEOSGS::SetSessionSearchParameters] EOS_SessionSearch_SetParameter failed with result [%s]"), *LexToString(ResultCode));
+		}
+	}
 }
 
 void FSessionsEOSGS::SetSessionSearchSessionId(FSessionSearchHandleEOSGS& SessionSearchHandle, const FOnlineSessionIdHandle& SessionId)
@@ -906,6 +1104,8 @@ TOnlineAsyncOpHandle<FFindSessions> FSessionsEOSGS::FindSessions(FFindSessions::
 			return;
 		}
 
+		const FFindSessions::Params& OpParams = Op.GetParams();
+
 		// For a successful session, we'll get the search results
 		EOS_SessionSearch_GetSearchResultCountOptions GetSearchResultCountOptions = { };
 		GetSearchResultCountOptions.ApiVersion = EOS_SESSIONSEARCH_GETSEARCHRESULTCOUNT_API_LATEST;
@@ -916,6 +1116,9 @@ TOnlineAsyncOpHandle<FFindSessions> FSessionsEOSGS::FindSessions(FFindSessions::
 		CopySearchResultByIndexOptions.ApiVersion = EOS_SESSIONSEARCH_COPYSEARCHRESULTBYINDEX_API_LATEST;
 		static_assert(EOS_SESSIONSEARCH_COPYSEARCHRESULTBYINDEX_API_LATEST == 1, "EOS_SessionSearch_CopySearchResultByIndexOptions updated, check new fields");
 
+
+		TArray<TFuture<TDefaultErrorResult<FBuildSessionFromDetailsHandle>>> PendingSessionsBuilt;
+
 		for (int32 Index = 0; Index < NumSearchResults; Index++)
 		{
 			EOS_HSessionDetails SessionDetailsHandle = nullptr;
@@ -924,26 +1127,61 @@ TOnlineAsyncOpHandle<FFindSessions> FSessionsEOSGS::FindSessions(FFindSessions::
 			EOS_EResult CopySearchResultByIndexResult = EOS_SessionSearch_CopySearchResultByIndex(CurrentSessionSearchHandleEOSGS->SearchHandle, &CopySearchResultByIndexOptions, &SessionDetailsHandle);
 			if (CopySearchResultByIndexResult == EOS_EResult::EOS_Success)
 			{
-				// All session attributes get extracted with the SessionDetailsHandle
-				const TSharedRef<FSessionEOSGS>& NewSessionEOSGSRef = MakeShared<FSessionEOSGS>(SessionDetailsHandle);
+				TSharedRef<TPromise<TDefaultErrorResult<FBuildSessionFromDetailsHandle>>> BuildSessionPromise = MakeShared<TPromise<TDefaultErrorResult<FBuildSessionFromDetailsHandle>>>();
+				BuildSessionFromDetailsHandle({ OpParams.LocalUserId, MakeShared<FSessionDetailsHandleEOSGS>(SessionDetailsHandle) })
+				.OnComplete([this, BuildSessionPromise](const TOnlineResult<FBuildSessionFromDetailsHandle>& Result) mutable
+				{
+					if (Result.IsOk())
+					{
+						CurrentSessionSearch->FoundSessions.Add(Result.GetOkValue().Session);
 
-				CurrentSessionSearch->FoundSessions.Add(NewSessionEOSGSRef);
+						TMap<FOnlineSessionIdHandle, TSharedRef<FSession>>& UserMap = SessionSearchResultsUserMap.FindOrAdd(Result.GetOkValue().LocalUserId);
+						UserMap.Emplace(Result.GetOkValue().Session->SessionId, Result.GetOkValue().Session);
 
-				TMap<FOnlineSessionIdHandle, TSharedRef<FSession>>& UserMap = SessionSearchResultsUserMap.FindOrAdd(Op.GetParams().LocalUserId);
-				UserMap.Emplace(NewSessionEOSGSRef->SessionId, NewSessionEOSGSRef);
+						BuildSessionPromise->EmplaceValue(Result.GetOkValue());
+					}
+					else
+					{
+						BuildSessionPromise->EmplaceValue(Result.GetErrorValue());
+					}
+				});
+			
+				PendingSessionsBuilt.Add(BuildSessionPromise->GetFuture());
 			}
 			else
 			{
 				UE_LOG(LogTemp, Error, TEXT("[FSessionsEOSGS::FindSessions] EOS_SessionSearch_CopySearchResultByIndex failed for index [%d] with result [%s]"), Index, *LexToString(CopySearchResultByIndexResult));
+
+				Op.SetError(Errors::FromEOSResult(CopySearchResultByIndexResult));
 			}
 		}
 
-		Op.SetResult(MoveTemp(*CurrentSessionSearch));
+		WhenAll(MoveTemp(PendingSessionsBuilt))
+			.Next([this, WeakOp = Op.AsWeak()](TArray<TDefaultErrorResult<FBuildSessionFromDetailsHandle>>&& Results) mutable
+		{
+			if (TOnlineAsyncOpPtr<FFindSessions> StrongOp = WeakOp.Pin())
+			{
+				for (TDefaultErrorResult<FBuildSessionFromDetailsHandle>& Result : Results)
+				{
+					if (Result.IsError())
+					{
+						// Store first encountered error to return as result.
+						StrongOp->SetError(MoveTemp(Result.GetErrorValue()));
+						return;
+					}
+				}
 
-		CurrentSessionSearch.Reset();
-		CurrentSessionSearchHandleEOSGS.Reset();
-	})
-	.Enqueue(GetSerialQueue());
+				StrongOp->SetResult(MoveTemp(*CurrentSessionSearch));
+
+				CurrentSessionSearch.Reset();
+				CurrentSessionSearchHandleEOSGS.Reset();
+			}
+		});
+	});
+
+	// TODO: Call BuildSessionFromDetailsHandle as many times as we have details handles
+
+	Op->Enqueue(GetSerialQueue());
 
 	return Op->GetHandle();
 }
@@ -1135,6 +1373,95 @@ TResult<TSharedRef<const FSession>, FOnlineError> FSessionsEOSGS::BuildSessionFr
 	}
 }
 
+TResult<TArray<EOS_ProductUserId>, FOnlineError> GetProductUserIdsFromEOSGSSession(const EOS_HSessionDetails& SessionDetailsHandle)
+{
+	TArray<EOS_ProductUserId> Result;
+
+	EOS_SessionDetails_CopyInfoOptions CopyInfoOptions = { };
+	CopyInfoOptions.ApiVersion = EOS_SESSIONDETAILS_COPYINFO_API_LATEST;
+	static_assert(EOS_SESSIONDETAILS_COPYINFO_API_LATEST == 1, "EOS_SessionDetails_CopyInfoOptions updated, check new fields");
+
+	EOS_SessionDetails_Info* SessionDetailsInfo = nullptr;
+	EOS_EResult CopyInfoResult = EOS_SessionDetails_CopyInfo(SessionDetailsHandle, &CopyInfoOptions, &SessionDetailsInfo);
+	if (CopyInfoResult == EOS_EResult::EOS_Success)
+	{
+		// We retrieve all the session attributes
+		EOS_SessionDetails_GetSessionAttributeCountOptions GetAttributeCountOptions = {};
+		GetAttributeCountOptions.ApiVersion = EOS_SESSIONDETAILS_GETSESSIONATTRIBUTECOUNT_API_LATEST;
+		static_assert(EOS_SESSIONDETAILS_GETSESSIONATTRIBUTECOUNT_API_LATEST == 1, "EOS_SessionDetails_GetSessionAttributeCountOptions updated, check new fields");
+		
+		uint32_t AttributeCount = EOS_SessionDetails_GetSessionAttributeCount(SessionDetailsHandle, &GetAttributeCountOptions);
+		for (uint32_t Index = 0; Index < AttributeCount; ++Index)
+		{
+			EOS_SessionDetails_CopySessionAttributeByIndexOptions CopyAttributeByIndexOptions = {};
+			CopyAttributeByIndexOptions.ApiVersion = EOS_SESSIONDETAILS_COPYSESSIONATTRIBUTEBYINDEX_API_LATEST;
+			static_assert(EOS_SESSIONDETAILS_COPYSESSIONATTRIBUTEBYINDEX_API_LATEST == 1, "EOS_SessionDetails_CopySessionAttributeByIndexOptions updated, check new fields");
+
+			CopyAttributeByIndexOptions.AttrIndex = Index;
+
+			EOS_SessionDetails_Attribute* Attribute = nullptr;
+			EOS_EResult CopyAttributeByIndexResult = EOS_SessionDetails_CopySessionAttributeByIndex(SessionDetailsHandle, &CopyAttributeByIndexOptions, &Attribute);
+			if (CopyAttributeByIndexResult == EOS_EResult::EOS_Success)
+			{
+				// We parse a single attribute
+				FString Key(Attribute->Data->Key);
+
+				// If the Key contains the ':' character, it will contain a user id
+				if (Key.Contains(TEXT(":")))
+				{
+					TArray<FString> KeyComponents;
+					Key.ParseIntoArray(KeyComponents, TEXT(":"));
+					const FString& PlayerIdStr = KeyComponents[0] == EOSGS_REGISTERED_PLAYERS.ToString() ? KeyComponents[1] : KeyComponents[0];
+					const EOS_ProductUserId ProductUserId = EOS_ProductUserId_FromString(TCHAR_TO_UTF8(*PlayerIdStr));
+					Result.AddUnique(ProductUserId);
+				}
+			}
+			else
+			{
+				UE_LOG(LogTemp, Warning, TEXT("[FSessionsEOSGS::BuildSessionFromDetailsHandle] EOS_SessionDetails_CopySessionAttributeByIndex failed with result [%s]"), *LexToString(CopyAttributeByIndexResult));
+
+				TResult<TArray<EOS_ProductUserId>, FOnlineError>(Errors::FromEOSResult(CopyAttributeByIndexResult));
+			}
+		}
+	}
+	else
+	{
+		UE_LOG(LogTemp, Warning, TEXT("[FSessionsEOSGS::BuildSessionFromDetailsHandle] EOS_SessionDetails_CopyInfo failed with result [%s]"), *LexToString(CopyInfoResult));
+
+		TResult<TArray<EOS_ProductUserId>, FOnlineError>(Errors::FromEOSResult(CopyInfoResult));
+	}
+
+	return TResult<TArray<EOS_ProductUserId>, FOnlineError>(Result);
+}
+
+TOnlineAsyncOpHandle<FBuildSessionFromDetailsHandle> FSessionsEOSGS::BuildSessionFromDetailsHandle(FBuildSessionFromDetailsHandle::Params&& Params)
+{
+	TOnlineAsyncOpRef<FBuildSessionFromDetailsHandle> Op = GetOp<FBuildSessionFromDetailsHandle>(MoveTemp(Params));
+	const FBuildSessionFromDetailsHandle::Params& OpParams = Op->GetParams();
+
+	// The first step in the process will be retrieving all ids in the session that need resolution
+	TResult<TArray<EOS_ProductUserId>, FOnlineError> GetIdsResult = GetProductUserIdsFromEOSGSSession(Params.SessionDetailsHandleEOSGS->SessionDetailsHandle);
+
+	if (GetIdsResult.IsOk())
+	{
+		Services.Get<FAuthEOSGS>()->ResolveAccountIds(Params.LocalUserId, GetIdsResult.GetOkValue())
+			.Next([this, WeakOp = Op->AsWeak(), Params = MoveTemp(Params)](TArray<FOnlineAccountIdHandle> ResolvedAccountIds) mutable
+		{
+			if (TOnlineAsyncOpPtr<FBuildSessionFromDetailsHandle> StrongOp = WeakOp.Pin())
+			{
+				// After all the ids are resolved, we can build the session safely
+				StrongOp->SetResult(FBuildSessionFromDetailsHandle::Result{ StrongOp->GetParams().LocalUserId, MakeShared<FSessionEOSGS>(Params.SessionDetailsHandleEOSGS->SessionDetailsHandle) });
+			}
+		});
+	}
+	else
+	{
+		Op->SetError(MoveTemp(GetIdsResult.GetErrorValue()));
+	}
+
+	return Op->GetHandle();
+}
+
 TOnlineAsyncOpHandle<FSendSessionInvite> FSessionsEOSGS::SendSessionInvite(FSendSessionInvite::Params&& Params)
 {
 	TOnlineAsyncOpRef<FSendSessionInvite> Op = GetOp<FSendSessionInvite>(MoveTemp(Params));
@@ -1373,31 +1700,26 @@ TOnlineAsyncOpHandle<FRemoveSessionMembers> FSessionsEOSGS::RemoveSessionMembers
 
 void FSessionsEOSGS::AppendSessionToPacket(FNboSerializeToBuffer& Packet, const FSessionLAN& Session)
 {
+	using namespace NboSerializerLANSvc;
 	using namespace NboSerializerEOSGSSvc;
 
-	// We won't save CurrentState in the packet as all advertised sessions will be Valid
+	SerializeToBuffer(Packet, Session);
 	SerializeToBuffer(Packet, Session.OwnerUserId);
 	SerializeToBuffer(Packet, Session.SessionId);
-	Packet << *Session.OwnerInternetAddr;
-
-	// TODO: Write session settings to packet, after SchemaVariant work
+	SerializeToBuffer(Packet, Session.SessionSettings.RegisteredPlayers);
+	SerializeToBuffer(Packet, Session.SessionSettings.SessionMembers);
 }
 
 void FSessionsEOSGS::ReadSessionFromPacket(FNboSerializeFromBuffer& Packet, FSessionLAN& Session)
 {
+	using namespace NboSerializerLANSvc;
 	using namespace NboSerializerEOSGSSvc;
 
+	SerializeFromBuffer(Packet, Session);
 	SerializeFromBuffer(Packet, Session.OwnerUserId);
 	SerializeFromBuffer(Packet, Session.SessionId);
-	Packet >> *Session.OwnerInternetAddr;
-
-	// We'll set the connect address for the remote session as a custom parameter, so it can be read in OnlineServices' GetResolvedConnectString
-	FCustomSessionSetting ConnectString;
-	ConnectString.Data.Set<FString>(Session.OwnerInternetAddr->ToString(true));
-	ConnectString.Visibility = ECustomSessionSettingVisibility::ViaOnlineService;
-	Session.SessionSettings.CustomSettings.Add(CONNECT_STRING_TAG, ConnectString);
-
-	// TODO:: Read session settings from packet, after SchemaVariant work
+	SerializeFromBuffer(Packet, Session.SessionSettings.RegisteredPlayers);
+	SerializeFromBuffer(Packet, Session.SessionSettings.SessionMembers);
 }
 
 /* UE::Online */ }
\ No newline at end of file
diff --git a/Engine/Plugins/Online/Experimental/OnlineServicesEOSGS/Source/Private/Online/SessionsEOSGSTypes.cpp b/Engine/Plugins/Online/Experimental/OnlineServicesEOSGS/Source/Private/Online/SessionsEOSGSTypes.cpp
new file mode 100644
--- /dev/null
+++ b/Engine/Plugins/Online/Experimental/OnlineServicesEOSGS/Source/Private/Online/SessionsEOSGSTypes.cpp
@@ -0,0 +1,139 @@
+// Copyright Epic Games, Inc. All Rights Reserved.
+
+#include "Online/SessionsEOSGSTypes.h"
+
+namespace UE::Online {
+
+EOS_EOnlineSessionPermissionLevel ToServiceType(const ESessionJoinPolicy& Value)
+{
+	switch (Value)
+	{
+	case ESessionJoinPolicy::Public:		return EOS_EOnlineSessionPermissionLevel::EOS_OSPF_PublicAdvertised;
+	case ESessionJoinPolicy::FriendsOnly:	return EOS_EOnlineSessionPermissionLevel::EOS_OSPF_JoinViaPresence;
+	case ESessionJoinPolicy::InviteOnly:	return EOS_EOnlineSessionPermissionLevel::EOS_OSPF_InviteOnly;
+	}
+
+	checkNoEntry();
+	return EOS_EOnlineSessionPermissionLevel::EOS_OSPF_InviteOnly;
+}
+
+ESessionJoinPolicy FromServiceType(const EOS_EOnlineSessionPermissionLevel& Value)
+{
+	switch (Value)
+	{
+	case EOS_EOnlineSessionPermissionLevel::EOS_OSPF_PublicAdvertised:	return ESessionJoinPolicy::Public;
+	case EOS_EOnlineSessionPermissionLevel::EOS_OSPF_JoinViaPresence:	return ESessionJoinPolicy::FriendsOnly;
+	case EOS_EOnlineSessionPermissionLevel::EOS_OSPF_InviteOnly:		return ESessionJoinPolicy::InviteOnly;
+	}
+
+	checkNoEntry();
+	return ESessionJoinPolicy::InviteOnly;
+}
+
+EOS_ESessionAttributeAdvertisementType ToServiceType(const ESchemaAttributeVisibility& Value)
+{
+	switch (Value)
+	{
+	case ESchemaAttributeVisibility::Private: return EOS_ESessionAttributeAdvertisementType::EOS_SAAT_DontAdvertise;
+	case ESchemaAttributeVisibility::Public: return EOS_ESessionAttributeAdvertisementType::EOS_SAAT_Advertise;
+	}
+
+	checkNoEntry();
+	return EOS_ESessionAttributeAdvertisementType::EOS_SAAT_DontAdvertise;
+}
+
+ESchemaAttributeVisibility FromServiceType(const EOS_ESessionAttributeAdvertisementType& Value)
+{
+	switch (Value)
+	{
+	case EOS_ESessionAttributeAdvertisementType::EOS_SAAT_DontAdvertise: return ESchemaAttributeVisibility::Private;
+	case EOS_ESessionAttributeAdvertisementType::EOS_SAAT_Advertise: return ESchemaAttributeVisibility::Public;
+	}
+
+	checkNoEntry();
+	return ESchemaAttributeVisibility::Private;
+}
+
+EOS_EOnlineComparisonOp ToServiceType(const ESchemaAttributeComparisonOp& Value)
+{
+	switch (Value)
+	{
+	case ESchemaAttributeComparisonOp::Equals:			return EOS_EComparisonOp::EOS_CO_EQUAL;
+	case ESchemaAttributeComparisonOp::NotEquals:			return EOS_EComparisonOp::EOS_CO_NOTEQUAL;
+	case ESchemaAttributeComparisonOp::GreaterThan:		return EOS_EComparisonOp::EOS_CO_GREATERTHAN;
+	case ESchemaAttributeComparisonOp::GreaterThanEquals:	return EOS_EComparisonOp::EOS_CO_GREATERTHANOREQUAL;
+	case ESchemaAttributeComparisonOp::LessThan:			return EOS_EComparisonOp::EOS_CO_LESSTHAN;
+	case ESchemaAttributeComparisonOp::LessThanEquals:	return EOS_EComparisonOp::EOS_CO_LESSTHANOREQUAL;
+	case ESchemaAttributeComparisonOp::Near:				return EOS_EComparisonOp::EOS_CO_DISTANCE;
+	case ESchemaAttributeComparisonOp::In:				return EOS_EComparisonOp::EOS_CO_ONEOF;
+	case ESchemaAttributeComparisonOp::NotIn:				return EOS_EComparisonOp::EOS_CO_NOTANYOF;
+	}
+
+	checkNoEntry();
+	return EOS_EComparisonOp::EOS_CO_EQUAL;
+}
+
+FSessionAttributeConverter<ESessionAttributeConversionType::ToService>::FSessionAttributeConverter(const FSchemaAttributeId& Key, const FSchemaVariant& Value)
+	: FSessionAttributeConverter(TPair<FSchemaAttributeId, FSchemaVariant>(Key, Value))
+{
+
+}
+
+FSessionAttributeConverter<ESessionAttributeConversionType::ToService>::FSessionAttributeConverter(const TPair<FSchemaAttributeId, FSchemaVariant>& InData)
+	: KeyConverterStorage(InData.Key.ToString())
+{
+	AttributeData.ApiVersion = EOS_SESSIONS_ATTRIBUTEDATA_API_LATEST;
+	static_assert(EOS_SESSIONS_ATTRIBUTEDATA_API_LATEST == 1, "EOS_Sessions_AttributeData updated, check new fields");
+
+	AttributeData.Key = KeyConverterStorage.Get();
+
+	switch (InData.Value.VariantType)
+	{
+	case ESchemaAttributeType::Bool:
+		AttributeData.ValueType = EOS_ESessionAttributeType::EOS_AT_BOOLEAN;
+		AttributeData.Value.AsBool = InData.Value.GetBoolean();
+		break;
+	case ESchemaAttributeType::Double:
+		AttributeData.ValueType = EOS_ESessionAttributeType::EOS_AT_DOUBLE;
+		AttributeData.Value.AsDouble = InData.Value.GetDouble();
+		break;
+	case ESchemaAttributeType::Int64:
+		AttributeData.ValueType = EOS_ESessionAttributeType::EOS_AT_INT64;
+		AttributeData.Value.AsInt64 = InData.Value.GetInt64();
+		break;
+	case ESchemaAttributeType::String:
+		AttributeData.ValueType = EOS_ESessionAttributeType::EOS_AT_STRING;
+		ValueConverterStorage.Emplace(InData.Value.GetString());
+		AttributeData.Value.AsUtf8 = ValueConverterStorage->Get();
+		break;
+	}
+}
+
+FSessionAttributeConverter<ESessionAttributeConversionType::FromService>::FSessionAttributeConverter(const EOS_Sessions_AttributeData& InData)
+{
+	FSchemaAttributeId AttributeId = FSchemaAttributeId(UTF8_TO_TCHAR(InData.Key));
+	FSchemaVariant VariantData;
+
+	switch (InData.ValueType)
+	{
+	case EOS_ESessionAttributeType::EOS_AT_BOOLEAN:
+		VariantData.Set(InData.Value.AsBool != 0);
+		break;
+	case EOS_ESessionAttributeType::EOS_AT_INT64:
+		VariantData.Set(static_cast<int64>(InData.Value.AsInt64));
+		break;
+	case EOS_ESessionAttributeType::EOS_AT_DOUBLE:
+		VariantData.Set(InData.Value.AsDouble);
+		break;
+	case EOS_ESessionAttributeType::EOS_AT_STRING:
+		VariantData.Set(UTF8_TO_TCHAR(InData.Value.AsUtf8));
+		break;
+	default:
+		checkNoEntry();
+		break;
+	}
+
+	AttributeData = TPair<FSchemaAttributeId, FSchemaVariant>(MoveTemp(AttributeId), MoveTemp(VariantData));
+}
+
+/* UE::Online */ }
\ No newline at end of file
diff --git a/Engine/Plugins/Online/Experimental/OnlineServicesEOSGS/Source/Public/Online/NboSerializerEOSGSSvc.h b/Engine/Plugins/Online/Experimental/OnlineServicesEOSGS/Source/Public/Online/NboSerializerEOSGSSvc.h
--- a/Engine/Plugins/Online/Experimental/OnlineServicesEOSGS/Source/Public/Online/NboSerializerEOSGSSvc.h
+++ b/Engine/Plugins/Online/Experimental/OnlineServicesEOSGS/Source/Public/Online/NboSerializerEOSGSSvc.h
@@ -8,6 +8,7 @@
 #include "Online/LobbiesEOSGS.h"
 #include "Online/SessionsEOSGS.h"
 #include "Online/CoreOnline.h"
+#include "Online/NboSerializerCommonSvc.h"
 
 /**
  * Serializes data in network byte order form into a buffer
@@ -32,6 +33,28 @@ inline void SerializeToBuffer(FNboSerializeToBuffer& Ar, const FOnlineSessionIdH
 	Ar.WriteBinary(Data.GetData(), Data.Num());
 }
 
+inline void SerializeToBuffer(FNboSerializeToBuffer& Packet, const FSessionMembersMap& SessionMembersMap)
+{
+	Packet << SessionMembersMap.Num();
+
+	for (const TPair<FOnlineAccountIdHandle, FSessionMember>& Entry : SessionMembersMap)
+	{
+		SerializeToBuffer(Packet, Entry.Key);
+		NboSerializerCommonSvc::SerializeToBuffer(Packet, Entry.Value);
+	}
+}
+
+inline void SerializeToBuffer(FNboSerializeToBuffer& Packet, const FRegisteredPlayersMap& RegisteredPlayersMap)
+{
+	Packet << RegisteredPlayersMap.Num();
+
+	for (const TPair<FOnlineAccountIdHandle, FRegisteredPlayer>& Entry : RegisteredPlayersMap)
+	{
+		SerializeToBuffer(Packet, Entry.Key);
+		NboSerializerCommonSvc::SerializeToBuffer(Packet, Entry.Value);
+	}
+}
+
 /** NboSerializeFromBuffer methods */
 
 inline void SerializeFromBuffer(FNboSerializeFromBuffer& Ar, FOnlineAccountIdHandle& UniqueId)
@@ -52,6 +75,40 @@ inline void SerializeFromBuffer(FNboSerializeFromBuffer& Ar, FOnlineSessionIdHan
 	SessionId = FOnlineSessionIdRegistryEOSGS::Get().FromReplicationData(Data);
 }
 
+inline void SerializeFromBuffer(FNboSerializeFromBuffer& Packet, FSessionMembersMap& SessionMembersMap)
+{
+	int32 NumEntries = 0;
+	Packet >> NumEntries;
+
+	for (int32 Index = 0; Index < NumEntries; ++Index)
+	{
+		FOnlineAccountIdHandle Key;
+		SerializeFromBuffer(Packet, Key);
+
+		FSessionMember Value;
+		NboSerializerCommonSvc::SerializeFromBuffer(Packet, Value);
+
+		SessionMembersMap.Emplace(Key, Value);
+	}
+}
+
+inline void SerializeFromBuffer(FNboSerializeFromBuffer& Packet, FRegisteredPlayersMap& RegisteredPlayersMap)
+{
+	int32 NumEntries = 0;
+	Packet >> NumEntries;
+
+	for (int32 Index = 0; Index < NumEntries; ++Index)
+	{
+		FOnlineAccountIdHandle Key;
+		SerializeFromBuffer(Packet, Key);
+
+		FRegisteredPlayer Value;
+		NboSerializerCommonSvc::SerializeFromBuffer(Packet, Value);
+
+		RegisteredPlayersMap.Emplace(Key, Value);
+	}
+}
+
 /* NboSerializerNullSvc */ }
 
 /* UE::Online */ }
\ No newline at end of file
diff --git a/Engine/Plugins/Online/Experimental/OnlineServicesEOSGS/Source/Public/Online/SessionsEOSGS.h b/Engine/Plugins/Online/Experimental/OnlineServicesEOSGS/Source/Public/Online/SessionsEOSGS.h
--- a/Engine/Plugins/Online/Experimental/OnlineServicesEOSGS/Source/Public/Online/SessionsEOSGS.h
+++ b/Engine/Plugins/Online/Experimental/OnlineServicesEOSGS/Source/Public/Online/SessionsEOSGS.h
@@ -34,8 +34,6 @@ private:
 	FOnlineSessionInviteIdRegistryEOSGS();
 };
 
-static FName EOS_SESSIONS_BUCKET_ID = TEXT("EOS_SESSIONS_BUCKET_ID");
-
 struct FSessionModificationHandleEOSGS : FNoncopyable
 {
 	EOS_HSessionModification ModificationHandle;
@@ -86,6 +84,8 @@ class FSessionEOSGS : public FSessionLAN
 public:
 	FSessionEOSGS() = default;
 	FSessionEOSGS(const FSessionEOSGS& InSession) = default;
+
+	/** This constructor should only be used by BuildSessionFromDetailsHandle, after all user ids in the session have been resolved. */
 	FSessionEOSGS(const EOS_HSessionDetails& SessionDetailsHandle);
 
 	static const FSessionEOSGS& Cast(const FSession& InSession);
@@ -97,7 +97,7 @@ public:
 
 struct FUpdateSessionImpl
 {
-	static constexpr TCHAR Name[] = TEXT("FUpdateSessionImpl");
+	static constexpr TCHAR Name[] = TEXT("UpdateSessionImpl");
 
 	struct Params
 	{
@@ -110,6 +110,29 @@ struct FUpdateSessionImpl
 	};
 };
 
+struct FBuildSessionFromDetailsHandle
+{
+	static constexpr TCHAR Name[] = TEXT("BuildSessionFromDetailsHandle");
+
+	struct Params
+	{
+		/** User which will drive the id resolution */
+		FOnlineAccountIdHandle LocalUserId;
+
+		/** EOS session details handle used to extract the data */
+		TSharedRef<FSessionDetailsHandleEOSGS> SessionDetailsHandleEOSGS;
+	};
+
+	struct Result
+	{
+		/** User which started the resolution operation */
+		FOnlineAccountIdHandle LocalUserId;
+
+		/** Session built from the details handle */
+		TSharedRef<FSession> Session;
+	};
+};
+
 class ONLINESERVICESEOSGS_API FSessionsEOSGS : public FSessionsLAN
 {
 public:
@@ -145,6 +168,8 @@ protected:
 	void SetPermissionLevel(EOS_HSessionModification& SessionModHandle, const ESessionJoinPolicy& JoinPolicy);
 	void SetBucketId(EOS_HSessionModification& SessionModHandle, const FString& NewBucketId);
 	void SetMaxPlayers(EOS_HSessionModification& SessionModHandle, const uint32& NewMaxPlayers);
+	void AddAttribute(EOS_HSessionModification& SessionModificationHandle, const FSchemaAttributeId& Key, const FCustomSessionSetting& Value);
+	void RemoveAttribute(EOS_HSessionModification& SessionModificationHandle, const FSchemaAttributeId& Key);
 
 	void SetSessionSearchMaxResults(FSessionSearchHandleEOSGS& SessionSearchHandle, uint32 MaxResults);
 	void SetSessionSearchParameters(FSessionSearchHandleEOSGS& SessionSearchHandle, TArray<FFindSessionsSearchFilter> Filters);
@@ -176,6 +201,11 @@ protected:
 	 */
 	TResult<TSharedRef<const FSession>, FOnlineError> BuildSessionFromUIEvent(const EOS_UI_EventId& UIEventId) const;
 
+	/**
+	 * Builds a session from an EOS Session Details Handle. Asynchronous due to the id resolution process
+	 */
+	TOnlineAsyncOpHandle<FBuildSessionFromDetailsHandle> BuildSessionFromDetailsHandle(FBuildSessionFromDetailsHandle::Params&& Params);
+
 private:
 	// FSessionsLAN
 	virtual void AppendSessionToPacket(FNboSerializeToBuffer& Packet, const FSessionLAN& Session) override;
@@ -191,4 +221,25 @@ protected:
 	TSharedPtr<FSessionSearchHandleEOSGS> CurrentSessionSearchHandleEOSGS;
 };
 
+namespace Meta {
+
+BEGIN_ONLINE_STRUCT_META(FUpdateSessionImpl::Params)
+	ONLINE_STRUCT_FIELD(FUpdateSessionImpl::Params, SessionModificationHandle)
+END_ONLINE_STRUCT_META()
+
+BEGIN_ONLINE_STRUCT_META(FUpdateSessionImpl::Result)
+END_ONLINE_STRUCT_META()
+
+BEGIN_ONLINE_STRUCT_META(FBuildSessionFromDetailsHandle::Params)
+	ONLINE_STRUCT_FIELD(FBuildSessionFromDetailsHandle::Params, LocalUserId),
+	ONLINE_STRUCT_FIELD(FBuildSessionFromDetailsHandle::Params, SessionDetailsHandleEOSGS)
+END_ONLINE_STRUCT_META()
+
+BEGIN_ONLINE_STRUCT_META(FBuildSessionFromDetailsHandle::Result)
+	ONLINE_STRUCT_FIELD(FBuildSessionFromDetailsHandle::Result, LocalUserId),
+	ONLINE_STRUCT_FIELD(FBuildSessionFromDetailsHandle::Result, Session)
+END_ONLINE_STRUCT_META()
+
+/* Meta*/ }
+
 /* UE::Online */ }
\ No newline at end of file
diff --git a/Engine/Plugins/Online/Experimental/OnlineServicesEOSGS/Source/Public/Online/SessionsEOSGSTypes.h b/Engine/Plugins/Online/Experimental/OnlineServicesEOSGS/Source/Public/Online/SessionsEOSGSTypes.h
new file mode 100644
--- /dev/null
+++ b/Engine/Plugins/Online/Experimental/OnlineServicesEOSGS/Source/Public/Online/SessionsEOSGSTypes.h
@@ -0,0 +1,72 @@
+// Copyright Epic Games, Inc. All Rights Reserved.
+
+#pragma once
+
+#include "Online/Sessions.h"
+
+#include "eos_sessions_types.h"
+
+namespace UE::Online {
+
+static FName EOSGS_ALLOW_NEW_MEMBERS = TEXT("EOSGS_ALLOW_NEW_MEMBERS");
+static FName EOSGS_ALLOW_UNREGISTERED_PLAYERS = TEXT("EOSGS_ALLOW_UNREGISTERED_PLAYERS");
+static FName EOSGS_ANTI_CHEAT_PROTECTED = TEXT("EOSGS_ANTI_CHEAT_PROTECTED");
+static FName EOSGS_IS_DEDICATED_SERVER_SESSION = TEXT("EOSGS_IS_DEDICATED_SERVER_SESSION");
+static FName EOSGS_PRESENCE_ENABLED = TEXT("EOSGS_PRESENCE_ENABLED");
+static FName EOSGS_SCHEMA_NAME = TEXT("EOSGS_SCHEMA_NAME");
+static FName EOSGS_SESSION_ID_OVERRIDE = TEXT("EOSGS_SESSION_ID_OVERRIDE");
+static FName EOSGS_REGISTERED_PLAYERS = TEXT("EOSGS_REGISTERED_PLAYERS");
+static FName EOSGS_REGISTERED_PLAYER_HAS_RESERVED_SLOT = TEXT("EOSGS_REGISTERED_PLAYER_HAS_RESERVED_SLOT");
+static FName EOSGS_REGISTERED_PLAYER_IS_IN_SESSION = TEXT("EOSGS_REGISTERED_PLAYER_IS_IN_SESSION");
+
+static FName EOSGS_BUCKET_ID = TEXT("EOSGS_BUCKET_ID");
+
+EOS_EOnlineSessionPermissionLevel ToServiceType(const ESessionJoinPolicy& Value);
+ESessionJoinPolicy FromServiceType(const EOS_EOnlineSessionPermissionLevel& Value);
+
+EOS_ESessionAttributeAdvertisementType ToServiceType(const ESchemaAttributeVisibility& Value);
+ESchemaAttributeVisibility FromServiceType(const EOS_ESessionAttributeAdvertisementType& Value);
+
+EOS_EOnlineComparisonOp ToServiceType(const ESchemaAttributeComparisonOp& Value);
+
+enum class ESessionAttributeConversionType
+{
+	ToService,
+	FromService
+};
+
+template <ESessionAttributeConversionType>
+class FSessionAttributeConverter
+{
+public:
+};
+
+template<>
+class FSessionAttributeConverter<ESessionAttributeConversionType::ToService>
+{
+public:
+	FSessionAttributeConverter(const FSchemaAttributeId& Key, const FSchemaVariant& Value);
+
+	FSessionAttributeConverter(const TPair<FSchemaAttributeId, FSchemaVariant>& InData);
+
+	const EOS_Sessions_AttributeData& GetAttributeData() const { return AttributeData; }
+
+private:
+	FTCHARToUTF8 KeyConverterStorage;
+	TOptional<FTCHARToUTF8> ValueConverterStorage; // TOptional because we'll only use it if the FSessionVariant type is FString
+	EOS_Sessions_AttributeData AttributeData;
+};
+
+template<>
+class FSessionAttributeConverter<ESessionAttributeConversionType::FromService>
+{
+public:
+	FSessionAttributeConverter(const EOS_Sessions_AttributeData& InData);
+
+	const TPair<FSchemaAttributeId, FSchemaVariant>& GetAttributeData() const { return AttributeData; }
+
+private:
+	TPair<FSchemaAttributeId, FSchemaVariant> AttributeData;
+};
+
+/* UE::Online */ }
\ No newline at end of file
diff --git a/Engine/Plugins/Online/Experimental/OnlineServicesNull/Source/Private/Online/OnlineServicesNull.cpp b/Engine/Plugins/Online/Experimental/OnlineServicesNull/Source/Private/Online/OnlineServicesNull.cpp
--- a/Engine/Plugins/Online/Experimental/OnlineServicesNull/Source/Private/Online/OnlineServicesNull.cpp
+++ b/Engine/Plugins/Online/Experimental/OnlineServicesNull/Source/Private/Online/OnlineServicesNull.cpp
@@ -101,7 +101,7 @@ TOnlineResult<FGetResolvedConnectString> FOnlineServicesNull::GetResolvedConnect
 					if (Session->SessionSettings.CustomSettings.Contains(CONNECT_STRING_TAG))
 					{
 						FGetResolvedConnectString::Result Result;
-						Result.ResolvedConnectString = Session->SessionSettings.CustomSettings[CONNECT_STRING_TAG].Data.Get<FString>();
+						Result.ResolvedConnectString = Session->SessionSettings.CustomSettings[CONNECT_STRING_TAG].Data.GetString();
 						return TOnlineResult<FGetResolvedConnectString>(Result);
 					}
 					else
diff --git a/Engine/Plugins/Online/Experimental/OnlineServicesNull/Source/Private/Online/SessionsNull.cpp b/Engine/Plugins/Online/Experimental/OnlineServicesNull/Source/Private/Online/SessionsNull.cpp
--- a/Engine/Plugins/Online/Experimental/OnlineServicesNull/Source/Private/Online/SessionsNull.cpp
+++ b/Engine/Plugins/Online/Experimental/OnlineServicesNull/Source/Private/Online/SessionsNull.cpp
@@ -32,31 +32,26 @@ FSessionsNull::FSessionsNull(FOnlineServicesNull& InServices)
 
 void FSessionsNull::AppendSessionToPacket(FNboSerializeToBuffer& Packet, const FSessionLAN& Session)
 {
+	using namespace NboSerializerLANSvc;
 	using namespace NboSerializerNullSvc;
 
-	// We won't save CurrentState in the packet as all advertised sessions will be Valid
+	SerializeToBuffer(Packet, Session);
 	SerializeToBuffer(Packet, Session.OwnerUserId);
 	SerializeToBuffer(Packet, Session.SessionId);
-	Packet << *Session.OwnerInternetAddr;
-
-	// TODO: Write session settings to packet, after SchemaVariant work
+	SerializeToBuffer(Packet, Session.SessionSettings.RegisteredPlayers);
+	SerializeToBuffer(Packet, Session.SessionSettings.SessionMembers);
 }
 
 void FSessionsNull::ReadSessionFromPacket(FNboSerializeFromBuffer& Packet, FSessionLAN& Session)
 {
+	using namespace NboSerializerLANSvc;
 	using namespace NboSerializerNullSvc;
 
+	SerializeFromBuffer(Packet, Session);
 	SerializeFromBuffer(Packet, Session.OwnerUserId);
 	SerializeFromBuffer(Packet, Session.SessionId);
-	Packet >> *Session.OwnerInternetAddr;
-
-	// We'll set the connect address for the remote session as a custom parameter, so it can be read in OnlineServices' GetResolvedConnectString
-	FCustomSessionSetting ConnectString;
-	ConnectString.Data.Set<FString>(Session.OwnerInternetAddr->ToString(true));
-	ConnectString.Visibility = ECustomSessionSettingVisibility::ViaOnlineService;
-	Session.SessionSettings.CustomSettings.Add(CONNECT_STRING_TAG, ConnectString);
-
-	// TODO:: Read session settings from packet, after SchemaVariant work
+	SerializeFromBuffer(Packet, Session.SessionSettings.RegisteredPlayers);
+	SerializeFromBuffer(Packet, Session.SessionSettings.SessionMembers);
 }
 
 /* UE::Online */ }
\ No newline at end of file
diff --git a/Engine/Plugins/Online/Experimental/OnlineServicesNull/Source/Public/Online/NboSerializerNullSvc.h b/Engine/Plugins/Online/Experimental/OnlineServicesNull/Source/Public/Online/NboSerializerNullSvc.h
--- a/Engine/Plugins/Online/Experimental/OnlineServicesNull/Source/Public/Online/NboSerializerNullSvc.h
+++ b/Engine/Plugins/Online/Experimental/OnlineServicesNull/Source/Public/Online/NboSerializerNullSvc.h
@@ -8,6 +8,7 @@
 #include "Online/LobbiesNull.h"
 #include "Online/SessionsNull.h"
 #include "Online/CoreOnline.h"
+#include "Online/NboSerializerCommonSvc.h"
 
 /**
  * Serializes data in network byte order form into a buffer
@@ -32,6 +33,28 @@ inline void SerializeToBuffer(FNboSerializeToBuffer& Ar, const FOnlineSessionIdH
 	Ar.WriteBinary(Data.GetData(), Data.Num());
 }
 
+inline void SerializeToBuffer(FNboSerializeToBuffer& Packet, const FSessionMembersMap& SessionMembersMap)
+{
+	Packet << SessionMembersMap.Num();
+
+	for (const TPair<FOnlineAccountIdHandle, FSessionMember>& Entry : SessionMembersMap)
+	{
+		SerializeToBuffer(Packet, Entry.Key);
+		NboSerializerCommonSvc::SerializeToBuffer(Packet, Entry.Value);
+	}
+}
+
+inline void SerializeToBuffer(FNboSerializeToBuffer& Packet, const FRegisteredPlayersMap& RegisteredPlayersMap)
+{
+	Packet << RegisteredPlayersMap.Num();
+
+	for (const TPair<FOnlineAccountIdHandle, FRegisteredPlayer>& Entry : RegisteredPlayersMap)
+	{
+		SerializeToBuffer(Packet, Entry.Key);
+		NboSerializerCommonSvc::SerializeToBuffer(Packet, Entry.Value);
+	}
+}
+
 inline void SerializeToBuffer(FNboSerializeToBuffer& Ar, const TMap<FLobbyAttributeId, FLobbyVariant>& Map)
 {
 	Ar << Map.Num();
@@ -64,6 +87,40 @@ inline void SerializeFromBuffer(FNboSerializeFromBuffer& Ar, FOnlineSessionIdHan
 	SessionId = FOnlineSessionIdRegistryNull::Get().FromReplicationData(Data);
 }
 
+inline void SerializeFromBuffer(FNboSerializeFromBuffer& Packet, FSessionMembersMap& SessionMembersMap)
+{
+	int32 NumEntries = 0;
+	Packet >> NumEntries;
+
+	for (int32 Index = 0; Index < NumEntries; ++Index)
+	{
+		FOnlineAccountIdHandle Key;
+		SerializeFromBuffer(Packet, Key);
+
+		FSessionMember Value;
+		NboSerializerCommonSvc::SerializeFromBuffer(Packet, Value);
+
+		SessionMembersMap.Emplace(Key, Value);
+	}
+}
+
+inline void SerializeFromBuffer(FNboSerializeFromBuffer& Packet, FRegisteredPlayersMap& RegisteredPlayersMap)
+{
+	int32 NumEntries = 0;
+	Packet >> NumEntries;
+
+	for (int32 Index = 0; Index < NumEntries; ++Index)
+	{
+		FOnlineAccountIdHandle Key;
+		SerializeFromBuffer(Packet, Key);
+
+		FRegisteredPlayer Value;
+		NboSerializerCommonSvc::SerializeFromBuffer(Packet, Value);
+
+		RegisteredPlayersMap.Emplace(Key, Value);
+	}
+}
+
 inline void SerializeFromBuffer(FNboSerializeFromBuffer& Ar, TMap<FLobbyAttributeId, FLobbyVariant>& Map)
 {
 	int32 Size;
diff --git a/Engine/Plugins/Online/Experimental/OnlineServicesOSSAdapter/Source/Private/Online/SessionsOSSAdapter.cpp b/Engine/Plugins/Online/Experimental/OnlineServicesOSSAdapter/Source/Private/Online/SessionsOSSAdapter.cpp
--- a/Engine/Plugins/Online/Experimental/OnlineServicesOSSAdapter/Source/Private/Online/SessionsOSSAdapter.cpp
+++ b/Engine/Plugins/Online/Experimental/OnlineServicesOSSAdapter/Source/Private/Online/SessionsOSSAdapter.cpp
@@ -17,17 +17,17 @@ namespace UE::Online {
 
 /** Auxiliary functions */
 
-EOnlineComparisonOp::Type GetV1SessionSearchComparisonOp(const ESessionsComparisonOp& InValue)
+EOnlineComparisonOp::Type GetV1SessionSearchComparisonOp(const ESchemaAttributeComparisonOp& InValue)
 {
 	switch (InValue)
 	{
-	case ESessionsComparisonOp::Equals:				return EOnlineComparisonOp::Equals;
-	case ESessionsComparisonOp::GreaterThan:		return EOnlineComparisonOp::GreaterThan;
-	case ESessionsComparisonOp::GreaterThanEquals:	return EOnlineComparisonOp::GreaterThanEquals;
-	case ESessionsComparisonOp::LessThan:			return EOnlineComparisonOp::LessThan;
-	case ESessionsComparisonOp::LessThanEquals:		return EOnlineComparisonOp::LessThanEquals;
-	case ESessionsComparisonOp::Near:				return EOnlineComparisonOp::Near;
-	case ESessionsComparisonOp::NotEquals:			return EOnlineComparisonOp::NotEquals;
+	case ESchemaAttributeComparisonOp::Equals:			return EOnlineComparisonOp::Equals;
+	case ESchemaAttributeComparisonOp::GreaterThan:		return EOnlineComparisonOp::GreaterThan;
+	case ESchemaAttributeComparisonOp::GreaterThanEquals:	return EOnlineComparisonOp::GreaterThanEquals;
+	case ESchemaAttributeComparisonOp::LessThan:			return EOnlineComparisonOp::LessThan;
+	case ESchemaAttributeComparisonOp::LessThanEquals:	return EOnlineComparisonOp::LessThanEquals;
+	case ESchemaAttributeComparisonOp::Near:				return EOnlineComparisonOp::Near;
+	case ESchemaAttributeComparisonOp::NotEquals:			return EOnlineComparisonOp::NotEquals;
 	}
 
 	checkNoEntry();
@@ -53,11 +53,162 @@ ESessionState GetV2SessionState(const EOnlineSessionState::Type InValue)
 	return ESessionState::Invalid;
 }
 
+FSchemaVariant GetV2SessionVariant(const FVariantData& InValue)
+{
+	FSchemaVariant Result;
+	
+	switch(InValue.GetType())
+	{
+	case EOnlineKeyValuePairDataType::Int32:
+	{
+		int32 Value;
+		InValue.GetValue(Value);
+		Result.Set((int64)Value);
+		break;
+	}
+	case EOnlineKeyValuePairDataType::UInt32:
+	{
+		uint32 Value;
+		InValue.GetValue(Value);
+		Result.Set((int64)Value);
+		break;
+	}
+	case EOnlineKeyValuePairDataType::Int64:
+	{
+		int64 Value;
+		InValue.GetValue(Value);
+		Result.Set(Value);
+		break;
+	}
+	case EOnlineKeyValuePairDataType::UInt64:
+	{
+		uint64 Value;
+		InValue.GetValue(Value);
+		Result.Set((int64)Value);
+		break;
+	}
+	case EOnlineKeyValuePairDataType::Float:
+	{
+		float Value;
+		InValue.GetValue(Value);
+		Result.Set((double)Value);
+		break;
+	}	
+	case EOnlineKeyValuePairDataType::Double:
+	{
+		double Value;
+		InValue.GetValue(Value);
+		Result.Set(Value);
+		break;
+	}
+	case EOnlineKeyValuePairDataType::String:
+	{
+		FString Value;
+		InValue.GetValue(Value);
+		Result.Set(Value);
+		break;
+	}
+	case EOnlineKeyValuePairDataType::Bool:
+	{
+		bool Value;
+		InValue.GetValue(Value);
+		Result.Set(Value);
+		break;
+	}
+	case EOnlineKeyValuePairDataType::Empty: // Intentional fallthrough
+	case EOnlineKeyValuePairDataType::Blob:	// Intentional fallthrough
+	case EOnlineKeyValuePairDataType::Json:	// Intentional fallthrough
+	case EOnlineKeyValuePairDataType::MAX:	// Intentional fallthrough
+	default:
+		UE_LOG(LogTemp, Warning, TEXT("[GetV2SessionVariant] FVariantData type not supported by FSessionsVariant. No value was set."));
+		break;
+	}
+
+	return Result;
+}
+
+FVariantData GetV1VariantData(const FSchemaVariant& InValue)
+{
+	FVariantData Result;
+
+	switch (InValue.VariantType)
+	{
+	case ESchemaAttributeType::Bool:
+		Result.SetValue(InValue.GetBoolean());
+		break;
+	case ESchemaAttributeType::Double:
+		Result.SetValue(InValue.GetDouble());
+		break;
+	case ESchemaAttributeType::Int64:
+		Result.SetValue(InValue.GetInt64());
+		break;
+	case ESchemaAttributeType::String:
+		Result.SetValue(InValue.GetString());
+		break;
+	default:
+		UE_LOG(LogTemp, Warning, TEXT("[GetV1VariantData] FSchemaVariant type not supported by FVariantData. No value was set."));
+		break;
+	}
+
+	return Result;
+}
+
+ESchemaAttributeVisibility GetV2SessionsAttributeVisibility(const EOnlineDataAdvertisementType::Type& InValue)
+{
+	switch (InValue)
+	{
+	case EOnlineDataAdvertisementType::DontAdvertise:			// Intentional fallthrough
+	case EOnlineDataAdvertisementType::ViaPingOnly:				return ESchemaAttributeVisibility::Private;
+	case EOnlineDataAdvertisementType::ViaOnlineService:		// Intentional fallthrough
+	case EOnlineDataAdvertisementType::ViaOnlineServiceAndPing:	return ESchemaAttributeVisibility::Public;
+	}
+
+	checkNoEntry();
+	return ESchemaAttributeVisibility::Private;
+}
+
+EOnlineDataAdvertisementType::Type GetV1OnlineDataAdvertisementType(const ESchemaAttributeVisibility& InValue)
+{
+	switch (InValue)
+	{
+	case ESchemaAttributeVisibility::Private:	return EOnlineDataAdvertisementType::DontAdvertise;
+	case ESchemaAttributeVisibility::Public:	return EOnlineDataAdvertisementType::ViaOnlineServiceAndPing;
+	}
+
+	checkNoEntry();
+	return EOnlineDataAdvertisementType::DontAdvertise;
+}
+
 FCustomSessionSettingsMap GetV2SessionSettings(const ::FSessionSettings& InSessionSettings)
 {
 	FCustomSessionSettingsMap Result;
 
-	// TODO: Pending SchemaVariant work
+	for (const TPair<FName, FOnlineSessionSetting>& Entry : InSessionSettings)
+	{
+		FCustomSessionSetting NewCustomSetting;
+		NewCustomSetting.Data = GetV2SessionVariant(Entry.Value.Data);
+		NewCustomSetting.Visibility = GetV2SessionsAttributeVisibility(Entry.Value.AdvertisementType);
+		NewCustomSetting.ID = Entry.Value.ID;
+
+		Result.Emplace(Entry.Key, NewCustomSetting);
+	}
+
+	return Result;
+}
+
+::FSessionSettings GetV1SessionSettings(const FCustomSessionSettingsMap& InSessionSettings)
+{
+	::FSessionSettings Result;
+
+	for (const TPair<FSchemaAttributeId, FCustomSessionSetting>& Entry : InSessionSettings)
+	{
+		FOnlineSessionSetting NewSessionSetting;
+		NewSessionSetting.Data = GetV1VariantData(Entry.Value.Data);
+		NewSessionSetting.AdvertisementType = GetV1OnlineDataAdvertisementType(Entry.Value.Visibility);
+		NewSessionSetting.ID = Entry.Value.ID;
+
+		Result.Add(Entry.Key, NewSessionSetting);
+	}
 
 	return Result;
 }
@@ -72,19 +223,32 @@ TSharedRef<FOnlineSessionSearch> BuildV1SessionSearch(const FFindSessions::Param
 
 	if (const FFindSessionsSearchFilter* PingBucketSize = Algo::FindBy(SearchParams.Filters, OSS_ADAPTER_SESSION_SEARCH_PING_BUCKET_SIZE, &FFindSessionsSearchFilter::Key))
 	{
-		Result->PingBucketSize = (int32)PingBucketSize->Value.Get<int64>();
+		Result->PingBucketSize = (int32)PingBucketSize->Value.GetInt64();
 	}
 
 	if (const FFindSessionsSearchFilter* PlatformHash = Algo::FindBy(SearchParams.Filters, OSS_ADAPTER_SESSION_SEARCH_PLATFORM_HASH, &FFindSessionsSearchFilter::Key))
 	{
-		Result->PlatformHash = (int32)PlatformHash->Value.Get<int64>();
+		Result->PlatformHash = (int32)PlatformHash->Value.GetInt64();
 	}
 
-	//for (const FFindSessionsSearchFilter& SearchFilter : SearchParams.Filters)
-	//{
-		// TODO: Pending SchemaVariant support
-		// Result->QuerySettings.Set(SearchFilter.Key, SearchFilter.Value, GetV1SessionSearchComparisonOp(SearchFilter.ComparisonOp));
-	//}
+	for (const FFindSessionsSearchFilter& SearchFilter : SearchParams.Filters)
+	{
+		switch (SearchFilter.Value.VariantType)
+		{
+		case ESchemaAttributeType::Bool:
+			Result->QuerySettings.Set<bool>(SearchFilter.Key, SearchFilter.Value.GetBoolean(), GetV1SessionSearchComparisonOp(SearchFilter.ComparisonOp));
+			break;
+		case ESchemaAttributeType::Double:
+			Result->QuerySettings.Set<double>(SearchFilter.Key, SearchFilter.Value.GetDouble(), GetV1SessionSearchComparisonOp(SearchFilter.ComparisonOp));
+			break;
+		case ESchemaAttributeType::Int64:
+			Result->QuerySettings.Set<uint64>(SearchFilter.Key, (uint64)SearchFilter.Value.GetInt64(), GetV1SessionSearchComparisonOp(SearchFilter.ComparisonOp));
+			break;
+		case ESchemaAttributeType::String:
+			Result->QuerySettings.Set<FString>(SearchFilter.Key, SearchFilter.Value.GetString(), GetV1SessionSearchComparisonOp(SearchFilter.ComparisonOp));
+			break;
+		}
+	}
 
 	return Result;
 }
@@ -241,8 +405,7 @@ void FSessionsOSSAdapter::Initialize()
 				{
 					if (const FCustomSessionSetting* Setting = SessionMember->MemberSettings.Find(Entry.Key))
 					{
-						// TODO: First part of the comparison commented until the SchemaVariant changes are applied
-						if (/*Setting->Data != Entry.Value.Data ||*/ Setting->Visibility != Entry.Value.Visibility || Setting->ID != Entry.Value.ID)
+						if (Setting->Data != Entry.Value.Data || Setting->Visibility != Entry.Value.Visibility || Setting->ID != Entry.Value.ID)
 						{
 							SessionMemberUpdate.UpdatedMemberSettings.Add(Entry);
 						}
@@ -682,7 +845,24 @@ TOnlineAsyncOpHandle<FStartMatchmaking> FSessionsOSSAdapter::StartMatchmaking(FS
 		{
 			FSessionMatchmakingUser NewMatchMakingUser { ServicesOSSAdapter.GetAccountIdRegistry().GetIdValue(LocalUser.Key).ToSharedRef() };
 
-			// TODO: Copy Attributes from session members parameters (Pending SchemaVariant)
+			for (const TPair<FSchemaAttributeId, FCustomSessionSetting>& Entry : LocalUser.Value.MemberSettings)
+			{
+				switch (Entry.Value.Data.VariantType)
+				{
+				case ESchemaAttributeType::Bool:
+					NewMatchMakingUser.Attributes.Emplace(Entry.Key.ToString(), FVariantData(Entry.Value.Data.GetBoolean()));
+					break;
+				case ESchemaAttributeType::Double:
+					NewMatchMakingUser.Attributes.Emplace(Entry.Key.ToString(), FVariantData(Entry.Value.Data.GetDouble()));
+					break;
+				case ESchemaAttributeType::Int64:
+					NewMatchMakingUser.Attributes.Emplace(Entry.Key.ToString(), FVariantData(Entry.Value.Data.GetInt64()));
+					break;
+				case ESchemaAttributeType::String:
+					NewMatchMakingUser.Attributes.Emplace(Entry.Key.ToString(), FVariantData(Entry.Value.Data.GetString()));
+					break;
+				}
+			}
 
 			MatchMakingUsers.Add(NewMatchMakingUser);
 		}
@@ -808,7 +988,7 @@ TOnlineAsyncOpHandle<FJoinSession> FSessionsOSSAdapter::JoinSession(FJoinSession
 
 		if (const FCustomSessionSetting* PingInMs = FoundSession->SessionSettings.CustomSettings.Find(OSS_ADAPTER_SESSIONS_PING_IN_MS))
 		{
-			SearchResult.PingInMs = (int32)PingInMs->Data.Get<int64>();
+			SearchResult.PingInMs = (int32)PingInMs->Data.GetInt64();
 		}
 
 		SearchResult.Session = BuildV1Session(FoundSession);
@@ -1037,28 +1217,27 @@ FOnlineSessionSettings FSessionsOSSAdapter::BuildV1Settings(const FSessionSettin
 	Result.bShouldAdvertise = InSessionSettings.bAllowNewMembers;
 	if (const FCustomSessionSetting* BuildUniqueId = InSessionSettings.CustomSettings.Find(OSS_ADAPTER_SESSIONS_BUILD_UNIQUE_ID))
 	{
-		Result.BuildUniqueId = BuildUniqueId->Data.Get<int64>();
+		Result.BuildUniqueId = BuildUniqueId->Data.GetInt64();
 	}
 	if (const FCustomSessionSetting* BuildUniqueId = InSessionSettings.CustomSettings.Find(OSS_ADAPTER_SESSIONS_USE_LOBBIES_IF_AVAILABLE))
 	{
-		Result.bUseLobbiesIfAvailable = BuildUniqueId->Data.Get<bool>();
+		Result.bUseLobbiesIfAvailable = BuildUniqueId->Data.GetBoolean();
 	}
 	if (const FCustomSessionSetting* BuildUniqueId = InSessionSettings.CustomSettings.Find(OSS_ADAPTER_SESSIONS_USE_LOBBIES_VOICE_CHAT_IF_AVAILABLE))
 	{
-		Result.bUseLobbiesVoiceChatIfAvailable = BuildUniqueId->Data.Get<bool>();
+		Result.bUseLobbiesVoiceChatIfAvailable = BuildUniqueId->Data.GetBoolean();
 	}
 	Result.bUsesPresence = InSessionSettings.bPresenceEnabled;
 	if (const FCustomSessionSetting* BuildUniqueId = InSessionSettings.CustomSettings.Find(OSS_ADAPTER_SESSIONS_USES_STATS))
 	{
-		Result.bUsesStats = BuildUniqueId->Data.Get<bool>();
+		Result.bUsesStats = BuildUniqueId->Data.GetBoolean();
 	}
 	Result.NumPrivateConnections = InSessionSettings.NumMaxPrivateConnections;
 	Result.NumPublicConnections = InSessionSettings.NumMaxPublicConnections;
 	Result.Settings.Add(OSS_ADAPTER_SESSIONS_SCHEMA_NAME, InSessionSettings.SchemaName.ToString());
 	Result.SessionIdOverride = InSessionSettings.SessionIdOverride;
 
-	// TODO: Pending SchemaVariant work
-	//Result.Settings =
+	Result.Settings = GetV1SessionSettings(InSessionSettings.CustomSettings);
 	
 	FOnlineServicesOSSAdapter& ServicesOSSAdapter = static_cast<FOnlineServicesOSSAdapter&>(Services);
 
@@ -1084,22 +1263,23 @@ void FSessionsOSSAdapter::WriteV2SessionSettingsFromV1Session(const FOnlineSessi
 	OutSettings.NumOpenPublicConnections = InSession->NumOpenPublicConnections;
 	FString SchemaNameStr;
 	InSession->SessionSettings.Settings.FindChecked(OSS_ADAPTER_SESSIONS_SCHEMA_NAME).Data.GetValue(SchemaNameStr);
-	OutSettings.SchemaName = FName(SchemaNameStr);
+	OutSettings.SchemaName = FSchemaId(SchemaNameStr);
 	OutSettings.SessionIdOverride = InSession->SessionSettings.SessionIdOverride;
 
-	// TODO: Pending SchemaVariant work
-	//Result.CustomSettings = 
+	OutSettings.CustomSettings = GetV2SessionSettings(InSession->SessionSettings.Settings);
 
 	FCustomSessionSetting BuildUniqueId;
-	BuildUniqueId.Data.Set<int64>(InSession->SessionSettings.BuildUniqueId);
+	BuildUniqueId.Data.Set((int64)InSession->SessionSettings.BuildUniqueId);
 	OutSettings.CustomSettings.Add(OSS_ADAPTER_SESSIONS_BUILD_UNIQUE_ID, BuildUniqueId);
 
 	FOnlineServicesOSSAdapter& ServicesOSSAdapter = static_cast<FOnlineServicesOSSAdapter&>(Services);
 
-	for (TPair<FUniqueNetIdRef, ::FSessionSettings> Entry : InSession->SessionSettings.MemberSettings)
+	for (const TPair<FUniqueNetIdRef, ::FSessionSettings>& Entry : InSession->SessionSettings.MemberSettings)
 	{
 		FOnlineAccountIdHandle SessionMemberId = ServicesOSSAdapter.GetAccountIdRegistry().FindOrAddHandle(Entry.Key);
-		FSessionMember SessionMember = FSessionMember(); // TODO: Pending SchemaVariant work
+
+		FSessionMember SessionMember;
+		SessionMember.MemberSettings = GetV2SessionSettings(Entry.Value);
 
 		OutSettings.SessionMembers.Emplace(SessionMemberId, SessionMember);
 	}
@@ -1154,7 +1334,7 @@ TArray<TSharedRef<const FSession>> FSessionsOSSAdapter::BuildV2SessionSearchResu
 		TSharedRef<FSession> FoundSession = BuildV2Session(&SearchResult.Session);
 
 		FCustomSessionSetting PingInMs;
-		PingInMs.Data.Set<int64>((int64)SearchResult.PingInMs);
+		PingInMs.Data.Set((int64)SearchResult.PingInMs);
 		FoundSession->SessionSettings.CustomSettings.Emplace(OSS_ADAPTER_SESSIONS_PING_IN_MS, PingInMs);
 
 		FoundSessions.Add(FoundSession);
