- Fix occurrences of FTCHARToUTF8 being constructed with FString, which results in the converter containing an unterminated string.

diff --git a/Engine/Plugins/Online/Experimental/OnlineServices/Source/OnlineServicesCommon/Public/Online/SessionsCommon.h b/Engine/Plugins/Online/Experimental/OnlineServices/Source/OnlineServicesCommon/Public/Online/SessionsCommon.h
--- a/Engine/Plugins/Online/Experimental/OnlineServices/Source/OnlineServicesCommon/Public/Online/SessionsCommon.h
+++ b/Engine/Plugins/Online/Experimental/OnlineServices/Source/OnlineServicesCommon/Public/Online/SessionsCommon.h
@@ -31,7 +31,7 @@ public:
 	virtual inline TArray<uint8> ToReplicationData(const FOnlineSessionIdHandle& Handle) const override
 	{
 		const FString IdValue = BasicRegistry.FindIdValue(Handle);
-		const FTCHARToUTF8 IdValueUtf8(IdValue);
+		const FTCHARToUTF8 IdValueUtf8(*IdValue);
 
 		TArray<uint8> ReplicationData;
 		ReplicationData.SetNumUninitialized(IdValueUtf8.Length());
diff --git a/Engine/Plugins/Online/Experimental/OnlineServicesEOS/Source/Private/Online/PresenceEOS.cpp b/Engine/Plugins/Online/Experimental/OnlineServicesEOS/Source/Private/Online/PresenceEOS.cpp
--- a/Engine/Plugins/Online/Experimental/OnlineServicesEOS/Source/Private/Online/PresenceEOS.cpp
+++ b/Engine/Plugins/Online/Experimental/OnlineServicesEOS/Source/Private/Online/PresenceEOS.cpp
@@ -207,7 +207,7 @@ TOnlineAsyncOpHandle<FUpdatePresence> FPresenceEOS::UpdatePresence(FUpdatePresen
 			EOS_PresenceModification_SetRawRichTextOptions RawRichTextOptions = { };
 			RawRichTextOptions.ApiVersion = EOS_PRESENCEMODIFICATION_SETRAWRICHTEXT_API_LATEST;
 
-			FTCHARToUTF8 Utf8RawRichText(*Params.Presence->StatusString);
+			const FTCHARToUTF8 Utf8RawRichText(*Params.Presence->StatusString);
 			RawRichTextOptions.RichText = Utf8RawRichText.Get();
 
 			EOS_EResult SetRawRichTextResult = EOS_PresenceModification_SetRawRichText(ChangeHandle, &RawRichTextOptions);
@@ -239,7 +239,7 @@ TOnlineAsyncOpHandle<FUpdatePresence> FPresenceEOS::UpdatePresence(FUpdatePresen
 				int32 CurrentIndex = 0;
 				for (const FString& RemovedProperty : RemovedProperties)
 				{
-					const FTCHARToUTF8& Utf8Key = Utf8Strings.Emplace_GetRef(RemovedProperty);
+					const FTCHARToUTF8& Utf8Key = Utf8Strings.Emplace_GetRef(*RemovedProperty);
 
 					EOS_PresenceModification_DataRecordId& RecordId = RecordIds.Emplace_GetRef();
 					RecordId.ApiVersion = EOS_PRESENCEMODIFICATION_DATARECORDID_API_LATEST;
@@ -278,8 +278,8 @@ TOnlineAsyncOpHandle<FUpdatePresence> FPresenceEOS::UpdatePresence(FUpdatePresen
 				int32 CurrentIndex = 0;
 				for (const TPair<FString, FPresenceVariant>& UpdatedProperty : Params.Presence->Properties)
 				{
-					const FTCHARToUTF8& Utf8Key = Utf8Strings.Emplace_GetRef(UpdatedProperty.Key);
-					const FTCHARToUTF8& Utf8Value = Utf8Strings.Emplace_GetRef(UpdatedProperty.Value); // TODO: Better serialization
+					const FTCHARToUTF8& Utf8Key = Utf8Strings.Emplace_GetRef(*UpdatedProperty.Key);
+					const FTCHARToUTF8& Utf8Value = Utf8Strings.Emplace_GetRef(*UpdatedProperty.Value); // TODO: Better serialization
 
 					EOS_Presence_DataRecord& Record = Records.Emplace_GetRef();
 					Record.ApiVersion = EOS_PRESENCE_DATARECORD_API_LATEST;
@@ -399,7 +399,7 @@ TOnlineAsyncOpHandle<FPartialUpdatePresence> FPresenceEOS::PartialUpdatePresence
 						EOS_PresenceModification_SetRawRichTextOptions RawRichTextOptions = { };
 						RawRichTextOptions.ApiVersion = EOS_PRESENCEMODIFICATION_SETRAWRICHTEXT_API_LATEST;
 
-						FTCHARToUTF8 Utf8RawRichText(*Params.Mutations.StatusString.GetValue());
+						const FTCHARToUTF8 Utf8RawRichText(*Params.Mutations.StatusString.GetValue());
 						RawRichTextOptions.RichText = Utf8RawRichText.Get();
 
 						EOS_EResult SetRawRichTextResult = EOS_PresenceModification_SetRawRichText(ChangeHandle, &RawRichTextOptions);
@@ -424,7 +424,7 @@ TOnlineAsyncOpHandle<FPartialUpdatePresence> FPresenceEOS::PartialUpdatePresence
 						int32 CurrentIndex = 0;
 						for (const FString& RemovedProperty : Params.Mutations.RemovedProperties)
 						{
-							const FTCHARToUTF8& Utf8Key = Utf8Strings.Emplace_GetRef(RemovedProperty);
+							const FTCHARToUTF8& Utf8Key = Utf8Strings.Emplace_GetRef(*RemovedProperty);
 
 							EOS_PresenceModification_DataRecordId& RecordId = RecordIds.Emplace_GetRef();
 							RecordId.ApiVersion = EOS_PRESENCEMODIFICATION_DATARECORDID_API_LATEST;
@@ -463,8 +463,8 @@ TOnlineAsyncOpHandle<FPartialUpdatePresence> FPresenceEOS::PartialUpdatePresence
 						int32 CurrentIndex = 0;
 						for (const TPair<FString, FPresenceVariant>& UpdatedProperty : Params.Mutations.UpdatedProperties)
 						{
-							const FTCHARToUTF8& Utf8Key = Utf8Strings.Emplace_GetRef(UpdatedProperty.Key);
-							const FTCHARToUTF8& Utf8Value = Utf8Strings.Emplace_GetRef(UpdatedProperty.Value); // TODO: Better serialization
+							const FTCHARToUTF8& Utf8Key = Utf8Strings.Emplace_GetRef(*UpdatedProperty.Key);
+							const FTCHARToUTF8& Utf8Value = Utf8Strings.Emplace_GetRef(*UpdatedProperty.Value); // TODO: Better serialization
 
 							EOS_Presence_DataRecord& Record = Records.Emplace_GetRef();
 							Record.ApiVersion = EOS_PRESENCE_DATARECORD_API_LATEST;
diff --git a/Engine/Plugins/Online/Experimental/OnlineServicesEOSGS/Source/Private/Online/AchievementsEOSGS.cpp b/Engine/Plugins/Online/Experimental/OnlineServicesEOSGS/Source/Private/Online/AchievementsEOSGS.cpp
--- a/Engine/Plugins/Online/Experimental/OnlineServicesEOSGS/Source/Private/Online/AchievementsEOSGS.cpp
+++ b/Engine/Plugins/Online/Experimental/OnlineServicesEOSGS/Source/Private/Online/AchievementsEOSGS.cpp
@@ -372,7 +372,7 @@ TOnlineAsyncOpHandle<FUnlockAchievements> FAchievementsEOSGS::UnlockAchievements
 		AchievementIdPtrs.Reserve(InAchievementIds.Num());
 		for (const FString& AchievementId : InAchievementIds)
 		{
-			FTCHARToUTF8& Converter = AchievementIdConverters.Emplace_GetRef(*AchievementId);
+			const FTCHARToUTF8& Converter = AchievementIdConverters.Emplace_GetRef(*AchievementId);
 			AchievementIdPtrs.Emplace(Converter.Get());
 		}
 
diff --git a/Engine/Plugins/Online/Experimental/OnlineServicesEOSGS/Source/Private/Online/AuthEOSGS.cpp b/Engine/Plugins/Online/Experimental/OnlineServicesEOSGS/Source/Private/Online/AuthEOSGS.cpp
--- a/Engine/Plugins/Online/Experimental/OnlineServicesEOSGS/Source/Private/Online/AuthEOSGS.cpp
+++ b/Engine/Plugins/Online/Experimental/OnlineServicesEOSGS/Source/Private/Online/AuthEOSGS.cpp
@@ -402,7 +402,7 @@ TDefaultErrorResultInternal<FEOSAuthLoginOptions> FEOSAuthLoginOptions::CreateIm
 
 	if (EnumHasAnyFlags(TranslatorTraits->Flags, EEOSAuthTranslationFlags::SetId))
 	{
-		FTCHARToUTF8 Converter(*CredentialsId);
+		const FTCHARToUTF8 Converter(*CredentialsId);
 		EOSAuthLoginOptions.IdUtf8.SetNumZeroed(Converter.Length() + 1);
 		FCStringAnsi::Strncpy(EOSAuthLoginOptions.IdUtf8.GetData(), Converter.Get(), EOSAuthLoginOptions.IdUtf8.Num());
 		EOSAuthLoginOptions.CredentialsData.Id = EOSAuthLoginOptions.IdUtf8.GetData();
@@ -416,7 +416,7 @@ TDefaultErrorResultInternal<FEOSAuthLoginOptions> FEOSAuthLoginOptions::CreateIm
 			return TDefaultErrorResultInternal<FEOSAuthLoginOptions>(Errors::InvalidCreds());
 		}
 
-		FTCHARToUTF8 Converter(*CredentialsToken.Get<FString>());
+		const FTCHARToUTF8 Converter(*CredentialsToken.Get<FString>());
 		EOSAuthLoginOptions.TokenUtf8.SetNumZeroed(Converter.Length() + 1);
 		FCStringAnsi::Strncpy(EOSAuthLoginOptions.TokenUtf8.GetData(), Converter.Get(), EOSAuthLoginOptions.TokenUtf8.Num());
 		EOSAuthLoginOptions.CredentialsData.Token = EOSAuthLoginOptions.TokenUtf8.GetData();
diff --git a/Engine/Plugins/Online/Experimental/OnlineServicesEOSGS/Source/Private/Online/SessionsEOSGS.cpp b/Engine/Plugins/Online/Experimental/OnlineServicesEOSGS/Source/Private/Online/SessionsEOSGS.cpp
--- a/Engine/Plugins/Online/Experimental/OnlineServicesEOSGS/Source/Private/Online/SessionsEOSGS.cpp
+++ b/Engine/Plugins/Online/Experimental/OnlineServicesEOSGS/Source/Private/Online/SessionsEOSGS.cpp
@@ -407,7 +407,7 @@ TOnlineAsyncOpHandle<FCreateSession> FSessionsEOSGS::CreateSession(FCreateSessio
 		CreateSessionModificationOptions.bSanctionsEnabled = !OpParams.SessionSettings.bAllowSanctionedPlayers;
 
 		const FCustomSessionSetting* BucketIdSetting = OpParams.SessionSettings.CustomSettings.Find(EOSGS_BUCKET_ID);
-		FTCHARToUTF8 BucketIdUtf8(BucketIdSetting != nullptr ? BucketIdSetting->Data.GetString() : FString());
+		const FTCHARToUTF8 BucketIdUtf8(BucketIdSetting ? *BucketIdSetting->Data.GetString() : nullptr);
 
 		if (BucketIdUtf8.Length())
 		{
@@ -417,13 +417,13 @@ TOnlineAsyncOpHandle<FCreateSession> FSessionsEOSGS::CreateSession(FCreateSessio
 		CreateSessionModificationOptions.LocalUserId = GetProductUserIdChecked(OpParams.LocalUserId);
 		CreateSessionModificationOptions.MaxPlayers = OpParams.SessionSettings.NumMaxPrivateConnections + OpParams.SessionSettings.NumMaxPublicConnections;
 
-		FTCHARToUTF8 SessionIdUtf8(OpParams.SessionSettings.SessionIdOverride);
+		const FTCHARToUTF8 SessionIdUtf8(*OpParams.SessionSettings.SessionIdOverride);
 		if (SessionIdUtf8.Length())
 		{
 			CreateSessionModificationOptions.SessionId = SessionIdUtf8.Get();
 		}
 
-		FTCHARToUTF8 SessionNameUtf8(OpParams.SessionName.ToString());
+		const FTCHARToUTF8 SessionNameUtf8(*OpParams.SessionName.ToString());
 		CreateSessionModificationOptions.SessionName = SessionNameUtf8.Get();
 
 		EOS_HSessionModification SessionModificationHandle = nullptr;
@@ -502,7 +502,7 @@ void FSessionsEOSGS::SetBucketId(EOS_HSessionModification& SessionModificationHa
 	Options.ApiVersion = EOS_SESSIONMODIFICATION_SETBUCKETID_API_LATEST;
 	static_assert(EOS_SESSIONMODIFICATION_SETBUCKETID_API_LATEST == 1, "EOS_SessionModification_SetBucketIdOptions updated, check new fields");
 
-	FTCHARToUTF8 BucketIdUtf8(*NewBucketId);
+	const FTCHARToUTF8 BucketIdUtf8(*NewBucketId);
 	Options.BucketId = BucketIdUtf8.Get();
 
 	EOS_EResult ResultCode = EOS_SessionModification_SetBucketId(SessionModificationHandle, &Options);
@@ -784,7 +784,7 @@ TOnlineAsyncOpHandle<FUpdateSession> FSessionsEOSGS::UpdateSession(FUpdateSessio
 		UpdateSessionModificationOptions.ApiVersion = EOS_SESSIONS_UPDATESESSIONMODIFICATION_API_LATEST;
 		static_assert(EOS_SESSIONS_UPDATESESSIONMODIFICATION_API_LATEST == 1, "EOS_Sessions_UpdateSessionModificationOptions updated, check new fields");
 
-		FTCHARToUTF8 SessionNameUtf8(OpParams.SessionName.ToString());
+		const FTCHARToUTF8 SessionNameUtf8(*OpParams.SessionName.ToString());
 		UpdateSessionModificationOptions.SessionName = SessionNameUtf8.Get();
 
 		EOS_HSessionModification SessionModificationHandle = nullptr;
@@ -901,7 +901,7 @@ TOnlineAsyncOpHandle<FLeaveSession> FSessionsEOSGS::LeaveSession(FLeaveSession::
 		DestroySessionOptions.ApiVersion = EOS_SESSIONS_DESTROYSESSION_API_LATEST;
 		static_assert(EOS_SESSIONS_DESTROYSESSION_API_LATEST == 1, "EOS_Sessions_DestroySessionOptions updated, check new fields");
 
-		FTCHARToUTF8 SessionNameUtf8(OpParams.SessionName.ToString());
+		const FTCHARToUTF8 SessionNameUtf8(*OpParams.SessionName.ToString());
 		DestroySessionOptions.SessionName = SessionNameUtf8.Get();
 
 		EOS_Async(EOS_Sessions_DestroySession, SessionsHandle, DestroySessionOptions, MoveTemp(Promise));
@@ -974,7 +974,7 @@ void FSessionsEOSGS::SetSessionSearchSessionId(FSessionSearchHandleEOSGS& Sessio
 	Options.ApiVersion = EOS_SESSIONSEARCH_SETSESSIONID_API_LATEST;
 	static_assert(EOS_SESSIONSEARCH_SETSESSIONID_API_LATEST == 1, "EOS_SessionSearch_SetSessionIdOptions updated, check new fields");
 
-	FTCHARToUTF8 SessionIdUtf8(FOnlineSessionIdRegistryEOSGS::Get().ToLogString(SessionId));
+	const FTCHARToUTF8 SessionIdUtf8(*FOnlineSessionIdRegistryEOSGS::Get().ToLogString(SessionId));
 	Options.SessionId = SessionIdUtf8.Get();
 
 	EOS_EResult ResultCode = EOS_SessionSearch_SetSessionId(SessionSearchHandle.SearchHandle, &Options);
@@ -1273,7 +1273,7 @@ TOnlineAsyncOpHandle<FJoinSession> FSessionsEOSGS::JoinSession(FJoinSession::Par
 
 		JoinSessionOptions.LocalUserId = GetProductUserIdChecked(OpParams.LocalUserId);
 
-		FTCHARToUTF8 SessionNameUtf8(OpParams.SessionName.ToString());
+		const FTCHARToUTF8 SessionNameUtf8(*OpParams.SessionName.ToString());
 		JoinSessionOptions.SessionName = SessionNameUtf8.Get();
 
 		const FSessionEOSGS& SessionEOSGS = FSessionEOSGS::Cast(*FoundSession);
@@ -1334,7 +1334,7 @@ TResult<TSharedRef<const FSession>, FOnlineError> FSessionsEOSGS::BuildSessionFr
 	CopySessionHandleByInviteIdOptions.ApiVersion = EOS_SESSIONS_COPYSESSIONHANDLEBYINVITEID_API_LATEST;
 	static_assert(EOS_SESSIONS_COPYSESSIONHANDLEBYINVITEID_API_LATEST == 1, "EOS_Sessions_CopySessionHandleByInviteIdOptions updated, check new fields");
 
-	FTCHARToUTF8 InviteIdUtf8(InInviteId);
+	const FTCHARToUTF8 InviteIdUtf8(*InInviteId);
 	CopySessionHandleByInviteIdOptions.InviteId = InviteIdUtf8.Get();
 
 	EOS_HSessionDetails SessionDetailsHandle;
@@ -1484,7 +1484,7 @@ TOnlineAsyncOpHandle<FSendSessionInvite> FSessionsEOSGS::SendSessionInvite(FSend
 
 		SendInviteOptions.LocalUserId = GetProductUserIdChecked(OpParams.LocalUserId);
 		
-		const FTCHARToUTF8 SessionNameUtf8(OpParams.SessionName.ToString());
+		const FTCHARToUTF8 SessionNameUtf8(*OpParams.SessionName.ToString());
 		SendInviteOptions.SessionName = SessionNameUtf8.Get();
 
 		SendInviteOptions.TargetUserId = GetProductUserIdChecked(OpParams.TargetUsers[0]); // TODO: Multiple user invitations using WhenAll like JoinLobbyImpl
@@ -1530,7 +1530,7 @@ TOnlineAsyncOpHandle<FRejectSessionInvite> FSessionsEOSGS::RejectSessionInvite(F
 		RejectInviteOptions.LocalUserId = GetProductUserIdChecked(OpParams.LocalUserId);
 
 		const FString& InviteIdStr = FOnlineSessionInviteIdRegistryEOSGS::Get().BasicRegistry.FindIdValue(OpParams.SessionInviteId);
-		const FTCHARToUTF8 InviteIdUtf8(InviteIdStr);
+		const FTCHARToUTF8 InviteIdUtf8(*InviteIdStr);
 		RejectInviteOptions.InviteId = InviteIdUtf8.Get();
 
 		EOS_Async(EOS_Sessions_RejectInvite, SessionsHandle, RejectInviteOptions, MoveTemp(Promise));
@@ -1600,7 +1600,7 @@ TOnlineAsyncOpHandle<FAddSessionMembers> FSessionsEOSGS::AddSessionMembers(FAddS
 
 		Options.PlayersToRegister = PlayersToRegister.GetData();
 
-		FTCHARToUTF8 SessionNameUtf8(OpParams.SessionName.ToString());
+		const FTCHARToUTF8 SessionNameUtf8(*OpParams.SessionName.ToString());
 		Options.SessionName = SessionNameUtf8.Get();
 
 		EOS_Async(EOS_Sessions_RegisterPlayers, SessionsHandle, Options, MoveTemp(Promise));
@@ -1668,7 +1668,7 @@ TOnlineAsyncOpHandle<FRemoveSessionMembers> FSessionsEOSGS::RemoveSessionMembers
 
 		Options.PlayersToUnregister = PlayersToUnregister.GetData();
 
-		FTCHARToUTF8 SessionNameUtf8(OpParams.SessionName.ToString());
+		const FTCHARToUTF8 SessionNameUtf8(*OpParams.SessionName.ToString());
 		Options.SessionName = SessionNameUtf8.Get();
 
 		EOS_Async(EOS_Sessions_UnregisterPlayers, SessionsHandle, Options, MoveTemp(Promise));
