Added missing EOS platform include
Horde issue 203022

diff --git a/Engine/Plugins/Online/Experimental/OnlineServicesEOSGS/Source/Public/Online/SessionsEOSGSTypes.h b/Engine/Plugins/Online/Experimental/OnlineServicesEOSGS/Source/Public/Online/SessionsEOSGSTypes.h
--- a/Engine/Plugins/Online/Experimental/OnlineServicesEOSGS/Source/Public/Online/SessionsEOSGSTypes.h
+++ b/Engine/Plugins/Online/Experimental/OnlineServicesEOSGS/Source/Public/Online/SessionsEOSGSTypes.h
@@ -4,6 +4,10 @@
 
 #include "Online/Sessions.h"
 
+#if defined(EOS_PLATFORM_BASE_FILE_NAME)
+#include EOS_PLATFORM_BASE_FILE_NAME
+#endif
+
 #include "eos_sessions_types.h"
 
 namespace UE::Online {
