Added joinability change logic to EOSGS Sessions
Added multiple session invite logic
Fix for UTF8 string conversions

diff --git a/Engine/Plugins/Online/Experimental/OnlineServicesEOSGS/Source/Private/Online/SessionsEOSGS.cpp b/Engine/Plugins/Online/Experimental/OnlineServicesEOSGS/Source/Private/Online/SessionsEOSGS.cpp
--- a/Engine/Plugins/Online/Experimental/OnlineServicesEOSGS/Source/Private/Online/SessionsEOSGS.cpp
+++ b/Engine/Plugins/Online/Experimental/OnlineServicesEOSGS/Source/Private/Online/SessionsEOSGS.cpp
@@ -395,7 +395,15 @@ TOnlineAsyncOpHandle<FCreateSession> FSessionsEOSGS::CreateSession(FCreateSessio
 			}
 		}
 
-		// TODO: Check if bucket id is set. Depends on if the method can be called successfully without it
+		// Check if the Bucket Id custom setting is set. EOS Sessions can not be created without it
+		if (!OpParams.SessionSettings.CustomSettings.Contains(EOSGS_BUCKET_ID))
+		{
+			UE_LOG(LogTemp, Warning, TEXT("[FSessionsEOSGS::CreateSession] Could not create session without Custom Setting 'EOSGS_BUCKET_ID' (FString) set."));
+
+			Op.SetError(Errors::InvalidParams());
+
+			return MakeFulfilledPromise<TDefaultErrorResult<FUpdateSessionImpl>>().GetFuture();
+		}
 
 		// After all initial checks, we start the session creation operations
 
@@ -450,7 +458,8 @@ TOnlineAsyncOpHandle<FCreateSession> FSessionsEOSGS::CreateSession(FCreateSessio
 
 		LocalSessionsByName.Add(OpParams.SessionName, NewSessionEOSGSRef);
 
-		FUpdateSessionImpl::Params UpdateSessionImplParams { MakeShared<FSessionModificationHandleEOSGS>(SessionModificationHandle) };
+		// Always update joinability on session creation
+		FUpdateSessionImpl::Params UpdateSessionImplParams{ MakeShared<FSessionModificationHandleEOSGS>(SessionModificationHandle), FUpdateSessionJoinabilityParams{ OpParams.SessionName, OpParams.SessionSettings.bAllowNewMembers } };
 		
 		return UpdateSessionImpl(MoveTemp(UpdateSessionImplParams));
 	})
@@ -466,10 +475,9 @@ TOnlineAsyncOpHandle<FCreateSession> FSessionsEOSGS::CreateSession(FCreateSessio
 			// TODO: Make additional local users join the created session
 
 			Op.SetResult(FCreateSession::Result{ FoundSession });
-
 		}
 		else
-		{	
+		{
 			// If creation failed, we remove the session objects
 			LocalSessionsByName.Remove(OpParams.SessionName);
 
@@ -481,6 +489,36 @@ TOnlineAsyncOpHandle<FCreateSession> FSessionsEOSGS::CreateSession(FCreateSessio
 	return Op->GetHandle();
 }
 
+void FSessionsEOSGS::SetJoinInProgressAllowed(EOS_HSessionModification& SessionModHandle, bool bIsJoinInProgressAllowed)
+{
+	EOS_SessionModification_SetJoinInProgressAllowedOptions Options = {};
+	Options.ApiVersion = EOS_SESSIONMODIFICATION_SETJOININPROGRESSALLOWED_API_LATEST;
+	static_assert(EOS_SESSIONMODIFICATION_SETJOININPROGRESSALLOWED_API_LATEST == 1, "EOS_SessionModification_SetJoinInProgressAllowedOptions updated, check new fields");
+
+	Options.bAllowJoinInProgress = bIsJoinInProgressAllowed;
+
+	EOS_EResult ResultCode = EOS_SessionModification_SetJoinInProgressAllowed(SessionModHandle, &Options);
+	if (ResultCode != EOS_EResult::EOS_Success)
+	{
+		UE_LOG(LogTemp, Warning, TEXT("EOS_SessionModification_SetJoinInProgressAllowed failed with result [%s]"), *LexToString(ResultCode));
+	}
+}
+
+void FSessionsEOSGS::SetInvitesAllowed(EOS_HSessionModification& SessionModHandle, bool bAreInvitesAllowed)
+{
+	EOS_SessionModification_SetInvitesAllowedOptions Options = {};
+	Options.ApiVersion = EOS_SESSIONMODIFICATION_SETINVITESALLOWED_API_LATEST;
+	static_assert(EOS_SESSIONMODIFICATION_SETINVITESALLOWED_API_LATEST == 1, "EOS_SessionModification_SetInvitesAllowedOptions updated, check new fields");
+
+	Options.bInvitesAllowed = bAreInvitesAllowed;
+
+	EOS_EResult ResultCode = EOS_SessionModification_SetInvitesAllowed(SessionModHandle, &Options);
+	if (ResultCode != EOS_EResult::EOS_Success)
+	{
+		UE_LOG(LogTemp, Warning, TEXT("EOS_SessionModification_SetInvitesAllowed failed with result [%s]"), *LexToString(ResultCode));
+	}
+}
+
 void FSessionsEOSGS::SetPermissionLevel(EOS_HSessionModification& SessionModificationHandle, const ESessionJoinPolicy& NewJoinPolicy)
 {
 	EOS_SessionModification_SetPermissionLevelOptions Options = { };
@@ -551,7 +589,7 @@ void FSessionsEOSGS::RemoveAttribute(EOS_HSessionModification& SessionModificati
 	Options.ApiVersion = EOS_SESSIONMODIFICATION_REMOVEATTRIBUTE_API_LATEST;
 	static_assert(EOS_SESSIONMODIFICATION_REMOVEATTRIBUTE_API_LATEST == 1, "EOS_SessionModification_RemoveAttributeOptions updated, check new fields");
 
-	FTCHARToUTF8 KeyUtf8(Key.ToString());
+	const FTCHARToUTF8 KeyUtf8(Key.ToString());
 	Options.Key = KeyUtf8.Get();
 
 	EOS_EResult ResultCode = EOS_SessionModification_RemoveAttribute(SessionModificationHandle, &Options);
@@ -565,8 +603,11 @@ void FSessionsEOSGS::WriteCreateSessionModificationHandle(EOS_HSessionModificati
 {
 	// TODO: We have the option to call EOS_SessionModification_SetHostAddress in EOS, useful if the session owner changes
 
-	// TODO: We either call Start or End session, EOS_SessionModification_SetJoinInProgressAllowed and EOS_SessionModification_SetInvitesAllowed depending on this value changing, or a custom setting.
-	// Investigation pending to see if there is overlay changes when the session is not joinable (started + JIP false, or InvitesAllowed false)
+	// We'll update this setting in the session modification step, and start or end the session accordingly to bock or allow join processes
+	SetJoinInProgressAllowed(SessionModificationHandle, SessionSettings.bAllowNewMembers);
+
+	// We'll also update invite permissions for the session
+	SetInvitesAllowed(SessionModificationHandle, SessionSettings.bAllowNewMembers);
 
 	// We won't copy bIsLANSession since it' irrelevant for EOS Sessions
 	AddAttribute(SessionModificationHandle, EOSGS_ALLOW_NEW_MEMBERS, { FSchemaVariant(SessionSettings.bAllowNewMembers), ESchemaAttributeVisibility::Public });
@@ -625,12 +666,14 @@ void FSessionsEOSGS::WriteUpdateSessionModificationHandle(EOS_HSessionModificati
 {
 	// TODO: We have the option to call EOS_SessionModification_SetHostAddress in EOS, useful if the session owner changes
 	
-	// TODO: We either call Start or End session, EOS_SessionModification_SetJoinInProgressAllowed and EOS_SessionModification_SetInvitesAllowed depending on this value changing, or a custom setting.
-	// Investigation pending to see if there is overlay changes when the session is not joinable (started + JIP false, or InvitesAllowed false)
 
 	if (NewSettings.bAllowNewMembers.IsSet())
 	{
-		AddAttribute(SessionModificationHandle, EOSGS_ALLOW_NEW_MEMBERS, { FSchemaVariant(NewSettings.bAllowNewMembers.GetValue()), ESchemaAttributeVisibility::Public });
+		// We'll update this setting in the session modification step, and start or end the session accordingly to bock or allow join processes
+		SetJoinInProgressAllowed(SessionModificationHandle, NewSettings.bAllowNewMembers.GetValue());
+
+		// We'll also update invite permissions for the session
+		SetInvitesAllowed(SessionModificationHandle, NewSettings.bAllowNewMembers.GetValue());
 	}
 
 	if (NewSettings.bAllowUnregisteredPlayers.IsSet())
@@ -801,7 +844,14 @@ TOnlineAsyncOpHandle<FUpdateSession> FSessionsEOSGS::UpdateSession(FUpdateSessio
 		// After creating the session modification handle, we'll add all the updated data to it
 		WriteUpdateSessionModificationHandle(SessionModificationHandle, OpParams.SessionName, OpParams.Mutations);
 
-		FUpdateSessionImpl::Params UpdateSessionImplParams{ MakeShared<FSessionModificationHandleEOSGS>(SessionModificationHandle) };
+		// Whether we update joinability or not will depend on inf bAllowNewMembers was set to a new value
+		FUpdateSessionImpl::Params UpdateSessionImplParams = 
+		{
+		MakeShared<FSessionModificationHandleEOSGS>(SessionModificationHandle),
+		OpParams.Mutations.bAllowNewMembers.IsSet() ?
+			FUpdateSessionJoinabilityParams{ OpParams.SessionName, OpParams.Mutations.bAllowNewMembers.GetValue() } :
+			TOptional<FUpdateSessionJoinabilityParams>() // If bAllowNewMembers is not set, the TOptional will be unset too
+		};
 
 		return UpdateSessionImpl(MoveTemp(UpdateSessionImplParams));
 	})
@@ -849,29 +899,178 @@ TFuture<TDefaultErrorResult<FUpdateSessionImpl>> FSessionsEOSGS::UpdateSessionIm
 	TFuture<TDefaultErrorResult<FUpdateSessionImpl>> Future = Promise.GetFuture();
 
 	EOS_Async(EOS_Sessions_UpdateSession, SessionsHandle, UpdateSessionOptions,
-	[this, Promise = MoveTemp(Promise)](const EOS_Sessions_UpdateSessionCallbackInfo* Result) mutable
+	[this, Promise = MoveTemp(Promise), Params = MoveTemp(Params)](const EOS_Sessions_UpdateSessionCallbackInfo* Result) mutable
 	{
-		if (Result->ResultCode != EOS_EResult::EOS_Success && Result->ResultCode != EOS_EResult::EOS_Sessions_OutOfSync)
+		// If we only change bAllowNewMembers, the session update will yield an EOS_NoChange result, but we still need to continue to the next step
+		if (Result->ResultCode != EOS_EResult::EOS_Success && Result->ResultCode != EOS_EResult::EOS_Sessions_OutOfSync && Result->ResultCode != EOS_EResult::EOS_NoChange)
 		{
+			UE_LOG(LogTemp, Warning, TEXT("EOS_Sessions_UpdateSession failed with result [%s]"), *LexToString(Result->ResultCode));
 			Promise.EmplaceValue(Errors::FromEOSResult(Result->ResultCode));
 			return;
 		}
 
-		if (TSharedRef<FSession>* FoundSession = LocalSessionsByName.Find(FName(Result->SessionName)))
+		// In session creation calls, we'll need to set the new id for the session
+		TSharedRef<FSession>& FoundSession = LocalSessionsByName.FindChecked(FName(UTF8_TO_TCHAR(Result->SessionName)));
+		if (!FoundSession->SessionId.IsValid())
 		{
-			// In session creation calls, we'll need to set the new id for the session
-			if (!FoundSession->Get().SessionId.IsValid())
-			{
-				FoundSession->Get().SessionId = CreateSessionId(FString(Result->SessionId));
-			}
+			FoundSession->SessionId = CreateSessionId(UTF8_TO_TCHAR(Result->SessionId));
 		}
 
-		Promise.EmplaceValue(FUpdateSessionImpl::Result { });
+		// After the successful general update, if indicated, we'll update the joinability
+		if (Params.UpdateJoinabilitySettings.IsSet())
+		{
+			UpdateSessionJoinabilityImpl(MoveTemp(Params.UpdateJoinabilitySettings.GetValue()))
+				.Next([this, Promise = MoveTemp(Promise)](TDefaultErrorResult<FUpdateSessionJoinabilityImpl>&& Result) mutable
+			{
+				if (Result.IsOk())
+				{
+					Promise.EmplaceValue(FUpdateSessionImpl::Result{ });
+				}
+				else
+				{
+					Promise.EmplaceValue(Result.GetErrorValue());
+				}
+			});
+		}
+		else
+		{
+			Promise.EmplaceValue(FUpdateSessionImpl::Result{ });
+		}
 	});
 
 	return Future;
 }
 
+TFuture<TDefaultErrorResult<FUpdateSessionJoinabilityImpl>> FSessionsEOSGS::UpdateSessionJoinabilityImpl(FUpdateSessionJoinabilityImpl::Params&& Params)
+{
+	TPromise<TDefaultErrorResult<FUpdateSessionJoinabilityImpl>> Promise;
+	TFuture<TDefaultErrorResult<FUpdateSessionJoinabilityImpl>> Future = Promise.GetFuture();
+
+	// We get the active session handle with the session name
+	EOS_Sessions_CopyActiveSessionHandleOptions CopyActiveSessionHandleOptions = {};
+	CopyActiveSessionHandleOptions.ApiVersion = EOS_SESSIONS_COPYACTIVESESSIONHANDLE_API_LATEST;
+	static_assert(EOS_SESSIONS_COPYACTIVESESSIONHANDLE_API_LATEST == 1, "EOS_Sessions_CopyActiveSessionHandleOptions updated, check new fields");
+
+	const FTCHARToUTF8 SessionNameUtf8(*Params.SessionName.ToString());
+	CopyActiveSessionHandleOptions.SessionName = SessionNameUtf8.Get();
+
+	EOS_HActiveSession ActiveSessionHandle;
+	EOS_EResult CopyActiveSessionHandleResult = EOS_Sessions_CopyActiveSessionHandle(SessionsHandle, &CopyActiveSessionHandleOptions, &ActiveSessionHandle);
+	if (CopyActiveSessionHandleResult != EOS_EResult::EOS_Success)
+	{
+		UE_LOG(LogTemp, Warning, TEXT("EOS_Sessions_CopyActiveSessionHandle failed with result [%s]"), *LexToString(CopyActiveSessionHandleResult));
+	}
+
+	// We get the active session info with the handle
+	EOS_ActiveSession_CopyInfoOptions CopyInfoOptions = {};
+	CopyInfoOptions.ApiVersion = EOS_ACTIVESESSION_COPYINFO_API_LATEST;
+	static_assert(EOS_ACTIVESESSION_COPYINFO_API_LATEST == 1, "EOS_ActiveSession_CopyInfoOptions updated, check new fields");
+
+	EOS_ActiveSession_Info* ActiveSessionInfo = nullptr;
+	EOS_EResult CopyInfoResult = EOS_ActiveSession_CopyInfo(ActiveSessionHandle, &CopyInfoOptions, &ActiveSessionInfo);
+	if (CopyInfoResult != EOS_EResult::EOS_Success)
+	{
+		UE_LOG(LogTemp, Warning, TEXT("EOS_ActiveSession_CopyInfo failed with result [%s]"), *LexToString(CopyInfoResult));
+	}
+
+	// If not, we start or end the session to make it joinable or not (as we set JIP to false at creation time)
+	if (Params.bAllowNewMembers)
+	{
+		// We check state. If bAllowNewMembers is true and session has not started, there's no need to do anything
+		if (ActiveSessionInfo->State == EOS_EOnlineSessionState::EOS_OSS_InProgress)
+		{
+			EOS_Sessions_EndSessionOptions Options = {};
+			Options.ApiVersion = EOS_SESSIONS_ENDSESSION_API_LATEST;
+			static_assert(EOS_SESSIONS_ENDSESSION_API_LATEST == 1, "EOS_Sessions_EndSessionOptions updated, check new fields");
+
+			Options.SessionName = SessionNameUtf8.Get();
+
+			EOS_Async(EOS_Sessions_EndSession, SessionsHandle, Options,
+				[this, Promise = MoveTemp(Promise)](const EOS_Sessions_EndSessionCallbackInfo* Result) mutable
+			{
+				if (Result->ResultCode != EOS_EResult::EOS_Success && Result->ResultCode != EOS_EResult::EOS_Sessions_OutOfSync)
+				{
+					UE_LOG(LogTemp, Warning, TEXT("EOS_Sessions_EndSession failed with result [%s]"), *LexToString(Result->ResultCode));
+					Promise.EmplaceValue(Errors::FromEOSResult(Result->ResultCode));
+					return;
+				}
+
+				Promise.EmplaceValue(FUpdateSessionJoinabilityImpl::Result{ });
+			});
+		}
+		else
+		{
+			Promise.EmplaceValue(FUpdateSessionJoinabilityImpl::Result{ });
+		}
+	}
+	else
+	{
+		// We check state. If bAllowNewMembers is false and the session has started, there's no need to do anything
+		if (ActiveSessionInfo->State == EOS_EOnlineSessionState::EOS_OSS_Pending || ActiveSessionInfo->State == EOS_EOnlineSessionState::EOS_OSS_Ended)
+		{
+			EOS_Sessions_StartSessionOptions Options = {};
+			Options.ApiVersion = EOS_SESSIONS_STARTSESSION_API_LATEST;
+			static_assert(EOS_SESSIONS_STARTSESSION_API_LATEST == 1, "EOS_Sessions_StartSessionOptions updated, check new fields");
+
+			Options.SessionName = SessionNameUtf8.Get();
+
+			EOS_Async(EOS_Sessions_StartSession, SessionsHandle, Options,
+				[this, Promise = MoveTemp(Promise)](const EOS_Sessions_StartSessionCallbackInfo* Result) mutable
+			{
+				if (Result->ResultCode != EOS_EResult::EOS_Success && Result->ResultCode != EOS_EResult::EOS_Sessions_OutOfSync)
+				{
+					UE_LOG(LogTemp, Warning, TEXT("EOS_Sessions_StartSession failed with result [%s]"), *LexToString(Result->ResultCode));
+					Promise.EmplaceValue(Errors::FromEOSResult(Result->ResultCode));
+					return;
+				}
+
+				Promise.EmplaceValue(FUpdateSessionJoinabilityImpl::Result{ });
+			});
+		}
+		else
+		{
+			Promise.EmplaceValue(FUpdateSessionJoinabilityImpl::Result{ });
+		}
+	}
+
+	return Future;
+}
+
+TOnlineAsyncOpHandle<FSendSingleSessionInviteImpl> FSessionsEOSGS::SendSingleSessionInviteImpl(FSendSingleSessionInviteImpl::Params&& Params)
+{
+	TOnlineAsyncOpRef<FSendSingleSessionInviteImpl> Op = GetOp<FSendSingleSessionInviteImpl>(MoveTemp(Params));
+	const FSendSingleSessionInviteImpl::Params& OpParams = Op->GetParams();
+
+	EOS_Sessions_SendInviteOptions SendInviteOptions = { };
+	SendInviteOptions.ApiVersion = EOS_SESSIONS_SENDINVITE_API_LATEST;
+	static_assert(EOS_SESSIONS_SENDINVITE_API_LATEST == 1, "EOS_Sessions_SendInviteOptions updated, check new fields");
+
+	SendInviteOptions.LocalUserId = GetProductUserIdChecked(OpParams.LocalUserId);
+
+	const FTCHARToUTF8 SessionNameUtf8(*OpParams.SessionName.ToString());
+	SendInviteOptions.SessionName = SessionNameUtf8.Get();
+
+	SendInviteOptions.TargetUserId = GetProductUserIdChecked(OpParams.TargetUserId);
+
+	EOS_Async(EOS_Sessions_SendInvite, SessionsHandle, SendInviteOptions,
+		[this, WeakOp = Op->AsWeak()](const EOS_Sessions_SendInviteCallbackInfo* Result) mutable
+	{
+		if (TOnlineAsyncOpPtr<FSendSingleSessionInviteImpl> StrongOp = WeakOp.Pin())
+		{
+			if (Result->ResultCode != EOS_EResult::EOS_Success)
+			{
+				UE_LOG(LogTemp, Warning, TEXT("EOS_Sessions_SendInvite failed with result [%s]"), *LexToString(Result->ResultCode));
+				StrongOp->SetError(Errors::FromEOSResult(Result->ResultCode));
+				return;
+			}
+
+			StrongOp->SetResult(FSendSingleSessionInviteImpl::Result{ });
+		}
+	});
+
+	return Op->GetHandle();
+}
+
 TOnlineAsyncOpHandle<FLeaveSession> FSessionsEOSGS::LeaveSession(FLeaveSession::Params&& Params)
 {
 	// LAN Sessions
@@ -1314,6 +1513,8 @@ TOnlineAsyncOpHandle<FJoinSession> FSessionsEOSGS::JoinSession(FJoinSession::Par
 				}
 			}
 
+			Op.SetResult(FJoinSession::Result{ *FoundSession });
+
 			FSessionJoined Event = { { OpParams.LocalUserId }, *FoundSession };
 
 			SessionEvents.OnSessionJoined.Broadcast(Event);
@@ -1478,31 +1679,49 @@ TOnlineAsyncOpHandle<FSendSessionInvite> FSessionsEOSGS::SendSessionInvite(FSend
 			return;
 		}
 
-		EOS_Sessions_SendInviteOptions SendInviteOptions = { };
-		SendInviteOptions.ApiVersion = EOS_SESSIONS_SENDINVITE_API_LATEST;
-		static_assert(EOS_SESSIONS_SENDINVITE_API_LATEST == 1, "EOS_Sessions_SendInviteOptions updated, check new fields");
-
-		SendInviteOptions.LocalUserId = GetProductUserIdChecked(OpParams.LocalUserId);
-		
-		const FTCHARToUTF8 SessionNameUtf8(*OpParams.SessionName.ToString());
-		SendInviteOptions.SessionName = SessionNameUtf8.Get();
+		TArray<TFuture<TDefaultErrorResult<FSendSingleSessionInviteImpl>>> PendingSessionInvites;
+		for (const FOnlineAccountIdHandle& TargetUser : OpParams.TargetUsers)
+		{
+			FSendSingleSessionInviteImpl::Params FSendSingleSessionInviteParams;
+			FSendSingleSessionInviteParams.LocalUserId = OpParams.LocalUserId;
+			FSendSingleSessionInviteParams.SessionName = OpParams.SessionName;
+			FSendSingleSessionInviteParams.TargetUserId = TargetUser;
+
+			TSharedRef<TPromise<TDefaultErrorResult<FSendSingleSessionInviteImpl>>> SessionInvitePromise = MakeShared<TPromise<TDefaultErrorResult<FSendSingleSessionInviteImpl>>>();
+			SendSingleSessionInviteImpl(MoveTemp(FSendSingleSessionInviteParams))
+				.OnComplete([SessionInvitePromise](const TOnlineResult<FSendSingleSessionInviteImpl>& Result)
+					{
+						if (Result.IsOk())
+						{
+							SessionInvitePromise->EmplaceValue(Result.GetOkValue());
+						}
+						else
+						{
+							SessionInvitePromise->EmplaceValue(Result.GetErrorValue());
+						}
+					});
 
-		SendInviteOptions.TargetUserId = GetProductUserIdChecked(OpParams.TargetUsers[0]); // TODO: Multiple user invitations using WhenAll like JoinLobbyImpl
+			PendingSessionInvites.Emplace(SessionInvitePromise->GetFuture());
+		}
 
-		EOS_Async(EOS_Sessions_SendInvite, SessionsHandle, SendInviteOptions, MoveTemp(Promise));
-		return;
-	})
-	.Then([this](TOnlineAsyncOp<FSendSessionInvite>& Op, const EOS_Sessions_SendInviteCallbackInfo* Result)
-	{
-		if (Result->ResultCode != EOS_EResult::EOS_Success)
+		WhenAll(MoveTemp(PendingSessionInvites))
+			.Next([WeakOp = Op.AsWeak()](TArray<TDefaultErrorResult<FSendSingleSessionInviteImpl>>&& Results) mutable
 		{
-			Op.SetError(Errors::FromEOSResult(Result->ResultCode));
-			return;
-		}
+			if (TOnlineAsyncOpPtr<FSendSessionInvite> StrongOp = WeakOp.Pin())
+			{
+				for (TDefaultErrorResult<FSendSingleSessionInviteImpl>& Result : Results)
+				{
+					if (Result.IsError())
+					{
+						StrongOp->SetError(MoveTemp(Result.GetErrorValue()));
+						return;
+					}
+				}
 
-		Op.SetResult(FSendSessionInvite::Result{ });
-	})
-	.Enqueue(GetSerialQueue()); // TODO: Use the parallel queue instead when possible
+				StrongOp->SetResult(FSendSessionInvite::Result{ });
+			}
+		});
+	});
 
 	return Op->GetHandle();
 }
diff --git a/Engine/Plugins/Online/Experimental/OnlineServicesEOSGS/Source/Public/Online/SessionsEOSGS.h b/Engine/Plugins/Online/Experimental/OnlineServicesEOSGS/Source/Public/Online/SessionsEOSGS.h
--- a/Engine/Plugins/Online/Experimental/OnlineServicesEOSGS/Source/Public/Online/SessionsEOSGS.h
+++ b/Engine/Plugins/Online/Experimental/OnlineServicesEOSGS/Source/Public/Online/SessionsEOSGS.h
@@ -20,7 +20,6 @@ class FOnlineSessionIdRegistryEOSGS : public FOnlineSessionIdRegistryLAN
 {
 public:
 	static FOnlineSessionIdRegistryEOSGS& Get();
-
 private:
 	FOnlineSessionIdRegistryEOSGS();
 };
@@ -95,6 +94,26 @@ public:
 	TSharedPtr<FSessionDetailsHandleEOSGS> SessionDetailsHandle;
 };
 
+struct FUpdateSessionJoinabilityImpl
+{
+	static constexpr TCHAR Name[] = TEXT("UpdateSessionJoinabilityImpl");
+
+	struct Params
+	{
+		/** Name for the session, needed to start or end it */
+		FName SessionName;
+
+		/** Whether players (registered or not) are accepted as new members in the session. */
+		bool bAllowNewMembers;
+	};
+
+	struct Result
+	{
+	};
+};
+
+typedef FUpdateSessionJoinabilityImpl::Params FUpdateSessionJoinabilityParams;
+
 struct FUpdateSessionImpl
 {
 	static constexpr TCHAR Name[] = TEXT("UpdateSessionImpl");
@@ -103,6 +122,27 @@ struct FUpdateSessionImpl
 	{
 		/** Handle for the session modification operation */
 		TSharedRef<FSessionModificationHandleEOSGS> SessionModificationHandle;
+
+		/** If set, it will use the values set in the struct to update the session's joinability */
+		TOptional<FUpdateSessionJoinabilityParams> UpdateJoinabilitySettings;
+	};
+
+	struct Result
+	{
+	};
+};
+
+struct FSendSingleSessionInviteImpl
+{
+	static constexpr TCHAR Name[] = TEXT("SendSingleSessionInviteImpl");
+
+	struct Params
+	{
+		FOnlineAccountIdHandle LocalUserId;
+
+		FName SessionName;
+
+		FOnlineAccountIdHandle TargetUserId;
 	};
 
 	struct Result
@@ -165,6 +205,8 @@ protected:
 	void HandleSessionInviteAccepted(const EOS_Sessions_SessionInviteAcceptedCallbackInfo* Data);
 	void HandleJoinSessionAccepted(const EOS_Sessions_JoinSessionAcceptedCallbackInfo* Data);
 
+	void SetJoinInProgressAllowed(EOS_HSessionModification& SessionModHandle, bool bIsJoinInProgressAllowed);
+	void SetInvitesAllowed(EOS_HSessionModification& SessionModHandle, bool bAreInvitesAllowed);
 	void SetPermissionLevel(EOS_HSessionModification& SessionModHandle, const ESessionJoinPolicy& JoinPolicy);
 	void SetBucketId(EOS_HSessionModification& SessionModHandle, const FString& NewBucketId);
 	void SetMaxPlayers(EOS_HSessionModification& SessionModHandle, const uint32& NewMaxPlayers);
@@ -180,14 +222,32 @@ protected:
 	 * Writes all values in the passed SessionSettings to the SessionModificationHandle
 	 */
 	void WriteCreateSessionModificationHandle(EOS_HSessionModification& SessionModificationHandle, const FSessionSettings& SessionSettings);
+
 	/**
 	 * Writes only the new values for all updated settings to the SessionModificationHandle
 	 */
 	void WriteUpdateSessionModificationHandle(EOS_HSessionModification& SessionModificationHandle, const FName& SessionName, const FSessionSettingsUpdate& NewSettings);
-	TFuture<TDefaultErrorResult<FUpdateSessionImpl>> UpdateSessionImpl(FUpdateSessionImpl::Params&& Params);
 
+	/**
+	 * Writes all relevant values set in the FindSessions parameters into the SessionSearchHandle
+	 */
 	void WriteSessionSearchHandle(FSessionSearchHandleEOSGS& SessionSearchHandle, const FFindSessions::Params& Params);
 
+	/**
+	 * Internal method used by both CreateSession and UpdateSession to process a session update at the API level
+	 */
+	TFuture<TDefaultErrorResult<FUpdateSessionImpl>> UpdateSessionImpl(FUpdateSessionImpl::Params&& Params);
+
+	/**
+	 * Internal method called after UpdateSessionImpl to update joinability options for a session
+	 */
+	TFuture<TDefaultErrorResult<FUpdateSessionJoinabilityImpl>> UpdateSessionJoinabilityImpl(FUpdateSessionJoinabilityImpl::Params&& Params);
+
+	/**
+	 * Internal method called by SendSessionInvite for every user
+	 */
+	TOnlineAsyncOpHandle<FSendSingleSessionInviteImpl> SendSingleSessionInviteImpl(FSendSingleSessionInviteImpl::Params&& Params);
+
 	static FOnlineSessionIdHandle CreateSessionId(const FString& SessionId);
 	FOnlineSessionInviteIdHandle CreateSessionInviteId(const FString& SessionInviteId) const;
 
@@ -223,6 +283,14 @@ protected:
 
 namespace Meta {
 
+BEGIN_ONLINE_STRUCT_META(FUpdateSessionJoinabilityImpl::Params)
+	ONLINE_STRUCT_FIELD(FUpdateSessionJoinabilityImpl::Params, SessionName),
+	ONLINE_STRUCT_FIELD(FUpdateSessionJoinabilityImpl::Params, bAllowNewMembers)
+END_ONLINE_STRUCT_META()
+
+BEGIN_ONLINE_STRUCT_META(FUpdateSessionJoinabilityImpl::Result)
+END_ONLINE_STRUCT_META()		
+		
 BEGIN_ONLINE_STRUCT_META(FUpdateSessionImpl::Params)
 	ONLINE_STRUCT_FIELD(FUpdateSessionImpl::Params, SessionModificationHandle)
 END_ONLINE_STRUCT_META()
@@ -230,6 +298,15 @@ END_ONLINE_STRUCT_META()
 BEGIN_ONLINE_STRUCT_META(FUpdateSessionImpl::Result)
 END_ONLINE_STRUCT_META()
 
+BEGIN_ONLINE_STRUCT_META(FSendSingleSessionInviteImpl::Params)
+	ONLINE_STRUCT_FIELD(FSendSingleSessionInviteImpl::Params, LocalUserId),
+	ONLINE_STRUCT_FIELD(FSendSingleSessionInviteImpl::Params, SessionName),
+	ONLINE_STRUCT_FIELD(FSendSingleSessionInviteImpl::Params, TargetUserId)
+END_ONLINE_STRUCT_META()
+
+BEGIN_ONLINE_STRUCT_META(FSendSingleSessionInviteImpl::Result)
+END_ONLINE_STRUCT_META()
+
 BEGIN_ONLINE_STRUCT_META(FBuildSessionFromDetailsHandle::Params)
 	ONLINE_STRUCT_FIELD(FBuildSessionFromDetailsHandle::Params, LocalUserId),
 	ONLINE_STRUCT_FIELD(FBuildSessionFromDetailsHandle::Params, SessionDetailsHandleEOSGS)
