Reorganisation of indexation in Sessions to adopt the one-object-per-session model
Removed ESessionState and references
Refactored most methods to operate with ids instead of session references
Refactored EOSGS invite handling to use new method (which first resolves all user ids)
Added needed state management to OSS Adapter

diff --git a/Engine/Plugins/Online/Experimental/OnlineServices/Source/OnlineServicesCommon/Private/Online/NboSerializerCommonSvc.cpp b/Engine/Plugins/Online/Experimental/OnlineServices/Source/OnlineServicesCommon/Private/Online/NboSerializerCommonSvc.cpp
--- a/Engine/Plugins/Online/Experimental/OnlineServices/Source/OnlineServicesCommon/Private/Online/NboSerializerCommonSvc.cpp
+++ b/Engine/Plugins/Online/Experimental/OnlineServices/Source/OnlineServicesCommon/Private/Online/NboSerializerCommonSvc.cpp
@@ -219,7 +219,6 @@ void SerializeFromBuffer(FNboSerializeFromBuffer& Packet, FSession& Session)
 {
 	// Session.OwnerUserId will be deserialized in implementations as user id types will vary
 	// Session.SessionId will be deserialized in implementations as user id types will vary
-	Session.CurrentState = ESessionState::Valid;
 	SerializeFromBuffer(Packet, Session.SessionSettings);
 }
 
diff --git a/Engine/Plugins/Online/Experimental/OnlineServices/Source/OnlineServicesCommon/Private/Online/SessionsCommon.cpp b/Engine/Plugins/Online/Experimental/OnlineServices/Source/OnlineServicesCommon/Private/Online/SessionsCommon.cpp
--- a/Engine/Plugins/Online/Experimental/OnlineServices/Source/OnlineServicesCommon/Private/Online/SessionsCommon.cpp
+++ b/Engine/Plugins/Online/Experimental/OnlineServices/Source/OnlineServicesCommon/Private/Online/SessionsCommon.cpp
@@ -41,7 +41,7 @@ namespace UE::Online {
 	{
 		FGetAllSessions::Result Result;
 
-		for (const TPair<FName, TSharedRef<FSession>>& SessionPair : LocalSessionsByName)
+		for (const TPair<FOnlineSessionIdHandle, TSharedRef<FSession>>& SessionPair : AllSessionsById)
 		{
 			Result.Sessions.Add(SessionPair.Value);
 		}
@@ -51,38 +51,28 @@ namespace UE::Online {
 
 	TOnlineResult<FGetSessionByName> FSessionsCommon::GetSessionByName(FGetSessionByName::Params&& Params) const
 	{
-		if (const TSharedRef<FSession>* FoundSession = LocalSessionsByName.Find(Params.LocalName))
+		TOnlineResult<FGetMutableSessionByName> GetMutableSessionByNameResult = GetMutableSessionByName({ Params.LocalName });
+		if (GetMutableSessionByNameResult.IsOk())
 		{
-			return TOnlineResult<FGetSessionByName>({ *FoundSession });
+			return TOnlineResult<FGetSessionByName>({ GetMutableSessionByNameResult.GetOkValue().Session });
 		}
 		else
 		{
-			return TOnlineResult<FGetSessionByName>(Errors::NotFound());
-		}
+			return TOnlineResult<FGetSessionByName>(GetMutableSessionByNameResult.GetErrorValue());
+		}	
 	}
 
 	TOnlineResult<FGetSessionById> FSessionsCommon::GetSessionById(FGetSessionById::Params&& Params) const
 	{
-		if (const TMap<FOnlineSessionIdHandle, TSharedRef<FSession>>* UserMap = SessionSearchResultsUserMap.Find(Params.LocalUserId))
+		TOnlineResult<FGetMutableSessionById> GetMutableSessionByIdResult = GetMutableSessionById({ Params.IdHandle });
+		if (GetMutableSessionByIdResult.IsOk())
 		{
-			if (const TSharedRef<FSession>* FoundSession = UserMap->Find(Params.IdHandle))
-			{
-				return TOnlineResult<FGetSessionById>({ *FoundSession });
-			}
+			return TOnlineResult<FGetSessionById>({ GetMutableSessionByIdResult.GetOkValue().Session });
 		}
-
-		if (const TMap<FOnlineSessionInviteIdHandle, TSharedRef<FSessionInvite>>* UserMap = SessionInvitesUserMap.Find(Params.LocalUserId))
+		else
 		{
-			for (const TPair<FOnlineSessionInviteIdHandle, TSharedRef<FSessionInvite>>& Entry : *UserMap)
-			{
-				if (Entry.Value->Session->SessionId == Params.IdHandle)
-				{
-					return TOnlineResult<FGetSessionById>({ Entry.Value->Session });
-				}
-			}
+			return TOnlineResult<FGetSessionById>(GetMutableSessionByIdResult.GetErrorValue());
 		}
-
-		return TOnlineResult<FGetSessionById>({ Errors::NotFound() });
 	}
 
 	TOnlineAsyncOpHandle<FCreateSession> FSessionsCommon::CreateSession(FCreateSession::Params&& Params)
@@ -307,9 +297,88 @@ namespace UE::Online {
 		return SessionEvents.OnUISessionJoinRequested;
 	}
 
+	TOnlineResult<FGetMutableSessionByName> FSessionsCommon::GetMutableSessionByName(FGetMutableSessionByName::Params&& Params) const
+	{
+		if (const FOnlineSessionIdHandle* SessionIdHandle = LocalSessionsByName.Find(Params.LocalName))
+		{
+			check(AllSessionsById.Contains(*SessionIdHandle));
+
+			return TOnlineResult<FGetMutableSessionByName>({ AllSessionsById.FindChecked(*SessionIdHandle) });
+		}
+		else
+		{
+			return TOnlineResult<FGetMutableSessionByName>(Errors::NotFound());
+		}
+	}
+
+	TOnlineResult<FGetMutableSessionById> FSessionsCommon::GetMutableSessionById(FGetMutableSessionById::Params&& Params) const
+	{
+		if (const TSharedRef<FSession>* FoundSession = AllSessionsById.Find(Params.IdHandle))
+		{
+			return TOnlineResult<FGetMutableSessionById>({ *FoundSession });
+		}
+		else
+		{
+			return TOnlineResult<FGetMutableSessionById>(Errors::NotFound());
+		}
+	}
+
+	void FSessionsCommon::ClearSessionByName(const FName& SessionName)
+	{
+		for (const TPair<FAccountId, TArray<FName>>& Entry : NamedSessionUserMap)
+		{
+			if (Entry.Value.Contains(SessionName))
+			{
+				return;
+			}
+		}
+
+		// If no references were found, we'll remove the named session entry
+		LocalSessionsByName.Remove(SessionName);
+	}
+
+	void FSessionsCommon::ClearSessionById(const FOnlineSessionIdHandle& SessionId)
+	{
+		// PresenceSessionsUserMap is not evaluated, since any session there would also be in LocalSessionsByName
+		for (const TPair<FName, FOnlineSessionIdHandle>& Entry : LocalSessionsByName)
+		{
+			if (Entry.Value == SessionId)
+			{
+				return;
+			}
+		}
+
+		for (const TPair<FAccountId, TMap<FOnlineSessionInviteIdHandle, TSharedRef<FSessionInvite>>>& Entry : SessionInvitesUserMap)
+		{
+			for (const TPair<FOnlineSessionInviteIdHandle, TSharedRef<FSessionInvite>>& InviteMap : Entry.Value)
+			{
+				if (InviteMap.Value->SessionId == SessionId)
+				{
+					return;
+				}
+			}			
+		}
+
+		for (const TPair<FAccountId, TArray<FOnlineSessionIdHandle>>& Entry : SearchResultsUserMap)
+		{
+			if (Entry.Value.Contains(SessionId))
+			{
+				return;
+			}
+		}
+
+		// If no references were found, we'll remove the session entry
+		AllSessionsById.Remove(SessionId);
+	}
+
 	TOnlineResult<FAddSessionMembers> FSessionsCommon::AddSessionMembersImpl(const FAddSessionMembers::Params& Params)
 	{
-		TSharedRef<FSession> FoundSession = LocalSessionsByName.FindChecked(Params.SessionName);
+		TOnlineResult<FGetMutableSessionByName> GetMutableSessionByNameResult = GetMutableSessionByName({ Params.SessionName });
+		if (GetMutableSessionByNameResult.IsError())
+		{
+			return TOnlineResult<FAddSessionMembers>(GetMutableSessionByNameResult.GetErrorValue());
+		}
+		TSharedRef<FSession> FoundSession = GetMutableSessionByNameResult.GetOkValue().Session;
 
 		FSessionSettings& SessionSettings = FoundSession->SessionSettings;
 
@@ -344,7 +413,7 @@ namespace UE::Online {
 			}
 			else if (Params.bRegisterPlayers) // If it wasn't and the parameters dictate it, we'll register them with a reserved slot
 			{
-				FRegisterPlayers::Params RegisterPlayerParams = { Params.SessionName, { SessionMemberId }, true };
+				FRegisterPlayers::Params RegisterPlayerParams = { Params.LocalUserId, Params.SessionName, { SessionMemberId }, true };
 
 				TOnlineResult<FRegisterPlayers> Result = RegisterPlayersImpl(RegisterPlayerParams);
 				if (Result.IsError())
@@ -374,7 +443,12 @@ namespace UE::Online {
 
 	TOnlineResult<FRemoveSessionMembers> FSessionsCommon::RemoveSessionMembersImpl(const FRemoveSessionMembers::Params& Params)
 	{
-		TSharedRef<FSession> FoundSession = LocalSessionsByName.FindChecked(Params.SessionName);
+		TOnlineResult<FGetMutableSessionByName> GetMutableSessionByNameResult = GetMutableSessionByName({ Params.SessionName });
+		if (GetMutableSessionByNameResult.IsError())
+		{
+			return TOnlineResult<FRemoveSessionMembers>(GetMutableSessionByNameResult.GetErrorValue());
+		}
+		TSharedRef<FSession> FoundSession = GetMutableSessionByNameResult.GetOkValue().Session;
 
 		FSessionSettings& SessionSettings = FoundSession->SessionSettings;
 
@@ -382,7 +456,7 @@ namespace UE::Online {
 		{
 			if (Params.bUnregisterPlayers) // If the parameters dictate it, well unregister the player as well
 			{
-				FUnregisterPlayers::Params UnregisterPlayersParams = { Params.SessionName, { SessionMemberId }, false };
+				FUnregisterPlayers::Params UnregisterPlayersParams = { Params.LocalUserId, Params.SessionName, { SessionMemberId }, false };
 
 				TOnlineResult<FUnregisterPlayers> Result = UnregisterPlayersImpl(UnregisterPlayersParams);
 				if (Result.IsError())
@@ -416,7 +490,12 @@ namespace UE::Online {
 
 	TOnlineResult<FRegisterPlayers> FSessionsCommon::RegisterPlayersImpl(const FRegisterPlayers::Params& Params)
 	{
-		TSharedRef<FSession> FoundSession = LocalSessionsByName.FindChecked(Params.SessionName);
+		TOnlineResult<FGetMutableSessionByName> GetMutableSessionByNameResult = GetMutableSessionByName({ Params.SessionName });
+		if (GetMutableSessionByNameResult.IsError())
+		{
+			return TOnlineResult<FRegisterPlayers>(GetMutableSessionByNameResult.GetErrorValue());
+		}
+		TSharedRef<FSession> FoundSession = GetMutableSessionByNameResult.GetOkValue().Session;
 
 		FSessionSettings& SessionSettings = FoundSession->SessionSettings;
 
@@ -455,7 +534,12 @@ namespace UE::Online {
 
 	TOnlineResult<FUnregisterPlayers> FSessionsCommon::UnregisterPlayersImpl(const FUnregisterPlayers::Params& Params)
 	{
-		TSharedRef<FSession> FoundSession = LocalSessionsByName.FindChecked(Params.SessionName);
+		TOnlineResult<FGetMutableSessionByName> GetMutableSessionByNameResult = GetMutableSessionByName({ Params.SessionName });
+		if (GetMutableSessionByNameResult.IsError())
+		{
+			return TOnlineResult<FUnregisterPlayers>(GetMutableSessionByNameResult.GetErrorValue());
+		}
+		TSharedRef<FSession> FoundSession = GetMutableSessionByNameResult.GetOkValue().Session;
 
 		FSessionSettings& SessionSettings = FoundSession->SessionSettings;
 
@@ -518,6 +602,13 @@ namespace UE::Online {
 			return Errors::InvalidParams();
 		}
 
+		if (!Params.LocalUsers.Contains(Params.LocalUserId))
+		{
+			UE_LOG(LogTemp, Warning, TEXT("[FSessionsCommon::CheckCreateSessionParams] Could not create session with name [%s]. Missing Session Member info for user [%s]"), *Params.SessionName.ToString(), *ToLogString(Params.LocalUserId));
+
+			return Errors::InvalidParams();
+		}
+
 		if (Params.SessionSettings.NumMaxPrivateConnections == 0 && Params.SessionSettings.NumMaxPublicConnections == 0)
 		{
 			UE_LOG(LogTemp, Warning, TEXT("[FSessionsCommon::CheckCreateSessionParams] Could not create session with name [%s] with no valid NumMaxPrivateConnections [%d] or NumMaxPublicConnections [%d]"), *Params.SessionName.ToString(), Params.SessionSettings.NumMaxPrivateConnections, Params.SessionSettings.NumMaxPublicConnections);
@@ -530,7 +621,8 @@ namespace UE::Online {
 
 	FOnlineError FSessionsCommon::CheckCreateSessionState(const FCreateSession::Params& Params)
 	{
-		if (TSharedRef<FSession>* FoundSession = LocalSessionsByName.Find(Params.SessionName))
+		TOnlineResult<FGetSessionByName> GetSessionByNameResult = GetSessionByName({ Params.SessionName });
+		if (GetSessionByNameResult.IsError())
 		{
 			UE_LOG(LogTemp, Warning, TEXT("[FSessionsCommon::CheckCreateSessionState] Could not create session with name [%s]. A session with that name already exists"), *Params.SessionName.ToString());
 
@@ -539,13 +631,17 @@ namespace UE::Online {
 
 		if (Params.SessionSettings.bPresenceEnabled)
 		{
-			for (const TPair<FName, TSharedRef<FSession>>& Entry : LocalSessionsByName)
+			for (const TPair<FName, FOnlineSessionIdHandle>& Entry : LocalSessionsByName)
 			{
-				if (Entry.Value->SessionSettings.bPresenceEnabled)
+				TOnlineResult<FGetSessionById> GetSessionByIdResult = GetSessionById({ Params.LocalUserId, Entry.Value });
+				if (GetSessionByIdResult.IsOk())
 				{
-					UE_LOG(LogTemp, Warning, TEXT("[FSessionsCommon::CheckCreateSessionState] Could not create session with bPresenceEnabled set to true when another already exists [%s]."), *Entry.Key.ToString());
+					if (GetSessionByIdResult.GetOkValue().Session->SessionSettings.bPresenceEnabled)
+					{
+						UE_LOG(LogTemp, Warning, TEXT("[FSessionsCommon::CheckCreateSessionState] Could not create session with bPresenceEnabled set to true when another already exists [%s]."), *Entry.Key.ToString());
 
-					return Errors::InvalidState();
+						return Errors::InvalidState();
+					}
 				}
 			}
 		}
@@ -572,9 +668,12 @@ namespace UE::Online {
 
 	FOnlineError FSessionsCommon::CheckUpdateSessionState(const FUpdateSession::Params& Params)
 	{
-		if (TSharedRef<FSession>* FoundSession = LocalSessionsByName.Find(Params.SessionName))
+		TOnlineResult<FGetSessionByName> GetSessionByNameResult = GetSessionByName({ Params.SessionName });
+		if (GetSessionByNameResult.IsOk())
 		{
-			if (!FoundSession->Get().SessionSettings.bIsDedicatedServerSession)
+			TSharedRef<const FSession> FoundSession = GetSessionByNameResult.GetOkValue().Session;
+
+			if (!FoundSession->SessionSettings.bIsDedicatedServerSession)
 			{
 				if (!Services.GetAuthInterface()->IsLoggedIn(Params.LocalUserId))
 				{
@@ -588,7 +687,7 @@ namespace UE::Online {
 		{
 			UE_LOG(LogTemp, Warning, TEXT("[FSessionsCommon::CheckUpdateSessionState] Could not update session with name [%s]. Session not found"), *Params.SessionName.ToString());
 
-			return Errors::NotFound();
+			return GetSessionByNameResult.GetErrorValue();
 		}
 
 		return Errors::Success();
@@ -617,7 +716,7 @@ namespace UE::Online {
 		}
 
 		// Ongoing search check
-		if (CurrentSessionSearch.IsValid())
+		if (const TSharedRef<TOnlineAsyncOp<FFindSessions>>* CurrentSessionSearchHandle = CurrentSessionSearchHandlesUserMap.Find(Params.LocalUserId))
 		{
 			UE_LOG(LogTemp, Warning, TEXT("[FSessionsCommon::CheckFindSessionsState] Could not find sessions, search already in progress"));
 
@@ -642,7 +741,8 @@ namespace UE::Online {
 	FOnlineError FSessionsCommon::CheckStartMatchmakingState(const FStartMatchmaking::Params& Params)
 	{
 		// Check if a session with that name already exists
-		if (TSharedRef<FSession>* FoundSession = LocalSessionsByName.Find(Params.SessionName))
+		TOptional<FOnlineError> Result = CheckSessionExistsByName(Params.LocalUserId, Params.SessionName);
+		if (!Result.IsSet()) // If CheckSessionExistsByName did not return an error, a session with that name already exists
 		{
 			UE_LOG(LogTemp, Warning, TEXT("[FSessionsCommon::CheckStartMatchmakingState] Could not join session with name [%s]. A session with that name already exists"), *Params.SessionName.ToString());
 
@@ -707,7 +807,8 @@ namespace UE::Online {
 	FOnlineError FSessionsCommon::CheckJoinSessionState(const FJoinSession::Params& Params)
 	{
 		// Check if a session with that name already exists
-		if (TSharedRef<FSession>* FoundSession = LocalSessionsByName.Find(Params.SessionName))
+		TOptional<FOnlineError> Result = CheckSessionExistsByName(Params.LocalUserId, Params.SessionName);
+		if (!Result.IsSet()) // If CheckSessionExistsByName did not return an error, a session with that name already exists
 		{
 			UE_LOG(LogTemp, Warning, TEXT("[FSessionsCommon::CheckJoinSessionState] Could not join session with name [%s]. A session with that name already exists"), *Params.SessionName.ToString());
 
@@ -744,13 +845,6 @@ namespace UE::Online {
 
 		const FSessionSettings& SessionSettings = FoundSession->SessionSettings;
 
-		if (FoundSession->CurrentState != ESessionState::Valid)
-		{
-			UE_LOG(LogTemp, Warning, TEXT("[FSessionsCommon::CheckJoinSessionState] Could not join session with invalid session state [%s]"), LexToString(FoundSession->CurrentState));
-
-			return Errors::InvalidParams();
-		}
-
 		for (const FAccountId& LocalUserId : LocalUserIds)
 		{
 			if (SessionSettings.SessionMembers.Contains(LocalUserId))
@@ -800,7 +894,7 @@ namespace UE::Online {
 
 	FOnlineError FSessionsCommon::CheckAddSessionMembersState(const FAddSessionMembers::Params& Params)
 	{
-		if (TOptional<FOnlineError> Result = CheckSessionExistsByName(Params.SessionName))
+		if (TOptional<FOnlineError> Result = CheckSessionExistsByName(Params.LocalUserId, Params.SessionName))
 		{
 			UE_LOG(LogTemp, Warning, TEXT("[FSessionsCommon::CheckAddSessionMembersState] Could not add session member to session with name [%s]. Session not found"), *Params.SessionName.ToString());
 
@@ -814,7 +908,7 @@ namespace UE::Online {
 
 	FOnlineError FSessionsCommon::CheckRemoveSessionMembersState(const FRemoveSessionMembers::Params& Params)
 	{
-		if (TOptional<FOnlineError> Result = CheckSessionExistsByName(Params.SessionName))
+		if (TOptional<FOnlineError> Result = CheckSessionExistsByName(Params.LocalUserId, Params.SessionName))
 		{
 			UE_LOG(LogTemp, Warning, TEXT("[FSessionsCommon::CheckRemoveSessionMembersState] Could not remove session member from session with name [%s]. Session not found"), *Params.SessionName.ToString());
 
@@ -827,11 +921,11 @@ namespace UE::Online {
 	FOnlineError FSessionsCommon::CheckLeaveSessionState(const FLeaveSession::Params& Params)
 	{
 		// User login check for main caller, session check
-		if (!LocalSessionsByName.Contains(Params.SessionName))
+		if (TOptional<FOnlineError> Result = CheckSessionExistsByName(Params.LocalUserId, Params.SessionName))
 		{
 			UE_LOG(LogTemp, Warning, TEXT("[FSessionsCommon::CheckLeaveSessionState] Could not leave session with name [%s]. Session not found"), *Params.SessionName.ToString());
 
-			return Errors::NotFound();
+			return Result.GetValue();
 		}
 
 		IAuthPtr Auth = Services.GetAuthInterface();
@@ -853,9 +947,12 @@ namespace UE::Online {
 	FOnlineError FSessionsCommon::CheckSendSessionInviteState(const FSendSessionInvite::Params& Params)
 	{
 		// User login check for main caller, session check
-		if (TSharedRef<FSession>* FoundSession = LocalSessionsByName.Find(Params.SessionName))
+		TOnlineResult<FGetSessionByName> GetSessionByNameResult = GetSessionByName({ Params.SessionName });
+		if (GetSessionByNameResult.IsOk())
 		{
-			if (!FoundSession->Get().SessionSettings.bIsDedicatedServerSession)
+			TSharedRef<const FSession> FoundSession = GetSessionByNameResult.GetOkValue().Session;
+
+			if (!FoundSession->SessionSettings.bIsDedicatedServerSession)
 			{
 				if (!Services.GetAuthInterface()->IsLoggedIn(Params.LocalUserId))
 				{
@@ -893,7 +990,7 @@ namespace UE::Online {
 
 	FOnlineError FSessionsCommon::CheckRegisterPlayersState(const FRegisterPlayers::Params& Params)
 	{
-		if (TOptional<FOnlineError> Result = CheckSessionExistsByName(Params.SessionName))
+		if (TOptional<FOnlineError> Result = CheckSessionExistsByName(Params.LocalUserId, Params.SessionName))
 		{
 			UE_LOG(LogTemp, Warning, TEXT("[FSessionsCommon::CheckRegisterPlayersState] Could not register players in session with name [%s]. Session not found"), *Params.SessionName.ToString());
 
@@ -907,7 +1004,7 @@ namespace UE::Online {
 
 	FOnlineError FSessionsCommon::CheckUnregisterPlayersState(const FUnregisterPlayers::Params& Params)
 	{
-		if(TOptional<FOnlineError> Result = CheckSessionExistsByName(Params.SessionName))
+		if(TOptional<FOnlineError> Result = CheckSessionExistsByName(Params.LocalUserId, Params.SessionName))
 		{
 			UE_LOG(LogTemp, Warning, TEXT("[FSessionsCommon::CheckUnregisterPlayersState] Could not unregister players from session with name [%s]. Session not found"), *Params.SessionName.ToString());
 
@@ -918,13 +1015,14 @@ namespace UE::Online {
 	}
 
 	// TODO: Have all Check methods return TOptional too, change call sites, and write Macro for repeating code structure
-	TOptional<FOnlineError> FSessionsCommon::CheckSessionExistsByName(const FName& SessionName)
+	TOptional<FOnlineError> FSessionsCommon::CheckSessionExistsByName(const FAccountId& LocalUserId, const FName& SessionName)
 	{
 		TOptional<FOnlineError> Result;
 
-		if (!LocalSessionsByName.Contains(SessionName))
+		TOnlineResult<FGetSessionByName> GetSessionByNameResult = GetSessionByName({ SessionName });
+		if (GetSessionByNameResult.IsError())
 		{
-			Result.Emplace(Errors::NotFound());
+			Result.Emplace(GetSessionByNameResult.GetErrorValue());
 		}
 
 		return Result;
diff --git a/Engine/Plugins/Online/Experimental/OnlineServices/Source/OnlineServicesCommon/Public/Online/SessionsCommon.h b/Engine/Plugins/Online/Experimental/OnlineServices/Source/OnlineServicesCommon/Public/Online/SessionsCommon.h
--- a/Engine/Plugins/Online/Experimental/OnlineServices/Source/OnlineServicesCommon/Public/Online/SessionsCommon.h
+++ b/Engine/Plugins/Online/Experimental/OnlineServices/Source/OnlineServicesCommon/Public/Online/SessionsCommon.h
@@ -127,6 +127,40 @@ public:
 	TOnlineBasicSessionInviteIdRegistry<FString> BasicRegistry;
 };
 
+struct FGetMutableSessionByName
+{
+	static constexpr TCHAR Name[] = TEXT("GetMutableSessionByName");
+
+	struct Params
+	{
+		FName LocalName;
+	};
+
+	struct Result
+	{
+		TSharedRef<FSession> Session;
+
+		Result() = delete; // cannot default construct due to TSharedRef
+	};
+};
+
+struct FGetMutableSessionById
+{
+	static constexpr TCHAR Name[] = TEXT("GetMutableSessionById");
+
+	struct Params
+	{
+		FOnlineSessionIdHandle IdHandle;
+	};
+
+	struct Result
+	{
+		TSharedRef<FSession> Session;
+
+		Result() = delete; // cannot default construct due to TSharedRef
+	};
+};
+
 class ONLINESERVICESCOMMON_API FSessionsCommon : public TOnlineComponent<ISessions>
 {
 public:
@@ -163,6 +197,12 @@ public:
 	virtual TOnlineEvent<void(const FUISessionJoinRequested&)> OnUISessionJoinRequested() override;
 
 protected:
+	TOnlineResult<FGetMutableSessionByName> GetMutableSessionByName(FGetMutableSessionByName::Params&& Params) const;
+	TOnlineResult<FGetMutableSessionById> GetMutableSessionById(FGetMutableSessionById::Params&& Params) const;
+
+	void ClearSessionByName(const FName& SessionName);
+	void ClearSessionById(const FOnlineSessionIdHandle& SessionId);
+
 	TOnlineResult<FAddSessionMembers> AddSessionMembersImpl(const FAddSessionMembers::Params& Params);
 	TOnlineResult<FRemoveSessionMembers> RemoveSessionMembersImpl(const FRemoveSessionMembers::Params& Params);
 
@@ -198,7 +238,7 @@ protected:
 	FOnlineError CheckUnregisterPlayersState(const FUnregisterPlayers::Params& Params);
 
 private:
-	TOptional<FOnlineError> CheckSessionExistsByName(const FName& SessionName);
+	TOptional<FOnlineError> CheckSessionExistsByName(const FAccountId& LocalUserId, const FName& SessionName);
 
 protected:
 
@@ -211,15 +251,46 @@ protected:
 		TOnlineEventCallable<void(const FUISessionJoinRequested&)> OnUISessionJoinRequested;
 	} SessionEvents;
 
+	/** Map of named sessions a user is part of, indexed by user */
+	TMap<FAccountId, TArray<FName>> NamedSessionUserMap;
+
 	/** Map of sessions that local users are part of, indexed by their local name */
-	TMap<FName, TSharedRef<FSession>> LocalSessionsByName;
-	/** Cache for the last set of session search results, mapped per user */
-	TMap<FAccountId, TMap<FOnlineSessionIdHandle, TSharedRef<FSession>>> SessionSearchResultsUserMap;
+	TMap<FName, FOnlineSessionIdHandle> LocalSessionsByName;
+
+	/** Map of sessions that local users have set as their presence session to appear in the platform UI. A user may not have set any session as their presence session. */
+	TMap<FAccountId, FOnlineSessionIdHandle> PresenceSessionsUserMap;
+
 	/** Cache for received session invites, mapped per user */
 	TMap<FAccountId, TMap<FOnlineSessionInviteIdHandle, TSharedRef<FSessionInvite>>> SessionInvitesUserMap;
 
-	TSharedPtr<FFindSessions::Result> CurrentSessionSearch;
-	TSharedPtr<TOnlineAsyncOp<FFindSessions>> CurrentSessionSearchHandle;
+	/** Cache for the last set of session search results, mapped per user */
+	TMap<FAccountId, TArray<FOnlineSessionIdHandle>> SearchResultsUserMap;
+
+	/** Handle for an ongoing session search operation, mapped per user */
+	TMap<FAccountId, TSharedRef<TOnlineAsyncOp<FFindSessions>>> CurrentSessionSearchHandlesUserMap;
+
+	/** Set of every distinct FSession found, indexed by Id */
+	TMap<FOnlineSessionIdHandle, TSharedRef<FSession>> AllSessionsById;
 };
 
+namespace Meta {
+
+BEGIN_ONLINE_STRUCT_META(FGetMutableSessionByName::Params)
+	ONLINE_STRUCT_FIELD(FGetMutableSessionByName::Params, LocalName)
+END_ONLINE_STRUCT_META()
+
+BEGIN_ONLINE_STRUCT_META(FGetMutableSessionByName::Result)
+	ONLINE_STRUCT_FIELD(FGetMutableSessionByName::Result, Session)
+END_ONLINE_STRUCT_META()
+
+BEGIN_ONLINE_STRUCT_META(FGetMutableSessionById::Params)
+	ONLINE_STRUCT_FIELD(FGetMutableSessionById::Params, IdHandle)
+END_ONLINE_STRUCT_META()
+
+BEGIN_ONLINE_STRUCT_META(FGetMutableSessionById::Result)
+	ONLINE_STRUCT_FIELD(FGetMutableSessionById::Result, Session)
+END_ONLINE_STRUCT_META()
+
+/* Meta*/ }
+
 /* UE::Online */ }
diff --git a/Engine/Plugins/Online/Experimental/OnlineServices/Source/OnlineServicesCommonEngineUtils/Private/Online/SessionsLAN.cpp b/Engine/Plugins/Online/Experimental/OnlineServices/Source/OnlineServicesCommonEngineUtils/Private/Online/SessionsLAN.cpp
--- a/Engine/Plugins/Online/Experimental/OnlineServices/Source/OnlineServicesCommonEngineUtils/Private/Online/SessionsLAN.cpp
+++ b/Engine/Plugins/Online/Experimental/OnlineServices/Source/OnlineServicesCommonEngineUtils/Private/Online/SessionsLAN.cpp
@@ -152,7 +152,6 @@ TOnlineAsyncOpHandle<FCreateSession> FSessionsLAN::CreateSession(FCreateSession:
 		}
 
 		TSharedRef<FSessionLAN> NewSessionLANRef = MakeShared<FSessionLAN>();
-		NewSessionLANRef->CurrentState = ESessionState::Valid;
 		NewSessionLANRef->OwnerUserId = OpParams.LocalUserId;
 		NewSessionLANRef->SessionId = FOnlineSessionIdRegistryLAN::GetChecked(Services.GetServicesProvider()).GetNextSessionId();
 		NewSessionLANRef->SessionSettings = OpParams.SessionSettings;
@@ -160,9 +159,20 @@ TOnlineAsyncOpHandle<FCreateSession> FSessionsLAN::CreateSession(FCreateSession:
 		// For LAN sessions, we'll add all the Session members manually instead of calling JoinSession since there is no API calls involved
 		NewSessionLANRef->SessionSettings.SessionMembers.Append(OpParams.LocalUsers);
 
-		LocalSessionsByName.Add(OpParams.SessionName, NewSessionLANRef);
+		// We save the local object for the session, and set up the appropriate references
+		AllSessionsById.Emplace(NewSessionLANRef->SessionId, NewSessionLANRef);
 
-		Op.SetResult(FCreateSession::Result{ NewSessionLANRef });
+		LocalSessionsByName.Emplace(OpParams.SessionName, NewSessionLANRef->SessionId);
+
+		NamedSessionUserMap.FindOrAdd(OpParams.LocalUserId).AddUnique(OpParams.SessionName);
+
+		if (NewSessionLANRef->SessionSettings.bPresenceEnabled)
+		{
+			FOnlineSessionIdHandle& PresenceSessionId = PresenceSessionsUserMap.FindOrAdd(OpParams.LocalUserId);
+			PresenceSessionId = NewSessionLANRef->SessionId;
+		}
+
+		Op.SetResult({ });
 	})
 	.Enqueue(GetSerialQueue());
 
@@ -186,7 +196,10 @@ TOnlineAsyncOpHandle<FUpdateSession> FSessionsLAN::UpdateSession(FUpdateSession:
 
 		// We'll check and update all settings one by one, with additional logic wherever required
 
-		TSharedRef<FSession>& FoundSession = LocalSessionsByName.FindChecked(OpParams.SessionName);
+		TOnlineResult<FGetMutableSessionByName> GetMutableSessionByNameResult = GetMutableSessionByName({ OpParams.SessionName });
+		check(GetMutableSessionByNameResult.IsOk());
+		
+		TSharedRef<FSession> FoundSession = GetMutableSessionByNameResult.GetOkValue().Session;
 		FSessionSettings& SessionSettings = FoundSession->SessionSettings;
 		const FSessionSettingsUpdate& UpdatedSettings = OpParams.Mutations;
 
@@ -211,9 +224,9 @@ TOnlineAsyncOpHandle<FUpdateSession> FSessionsLAN::UpdateSession(FUpdateSession:
 		SessionSettings += UpdatedSettings;
 
 		// We set the result and fire the event
-		Op.SetResult(FUpdateSession::Result{ FoundSession });
+		Op.SetResult({ });
 
-		FSessionUpdated SessionUpdatedEvent{ FoundSession, UpdatedSettings };
+		FSessionUpdated SessionUpdatedEvent{ OpParams.SessionName, UpdatedSettings };
 		SessionEvents.OnSessionUpdated.Broadcast(SessionUpdatedEvent);
 	})
 	.Enqueue(GetSerialQueue());
@@ -244,10 +257,9 @@ TOnlineAsyncOpHandle<FFindSessions> FSessionsLAN::FindSessions(FFindSessions::Pa
 			return;
 		}
 
-		CurrentSessionSearch = MakeShared<FFindSessions::Result>();
-		CurrentSessionSearchHandle = Op.AsShared();
+		CurrentSessionSearchHandlesUserMap.Emplace(OpParams.LocalUserId, Op.AsShared());
 
-		FindLANSessions();
+		FindLANSessions(OpParams.LocalUserId);
 	})
 	.Enqueue(GetSerialQueue());
 
@@ -292,17 +304,28 @@ TOnlineAsyncOpHandle<FJoinSession> FSessionsLAN::JoinSession(FJoinSession::Param
 		}
 
 		const TSharedRef<const FSession>& FoundSession = GetSessionByIdResult.GetOkValue().Session;
-		TSharedRef<FSession> NewSessionRef = ConstCastSharedRef<FSession>(FoundSession);
+		TSharedRef<FSession> NewSessionLANRef = ConstCastSharedRef<FSession>(FoundSession);
+
+		// We save the local object for the session, and set up the appropriate references
+		AllSessionsById.Emplace(NewSessionLANRef->SessionId, NewSessionLANRef);
+
+		LocalSessionsByName.Emplace(OpParams.SessionName, NewSessionLANRef->SessionId);
 
-		LocalSessionsByName.Add(OpParams.SessionName, NewSessionRef);
+		NamedSessionUserMap.FindOrAdd(OpParams.LocalUserId).AddUnique(OpParams.SessionName);
+
+		if (NewSessionLANRef->SessionSettings.bPresenceEnabled)
+		{
+			FOnlineSessionIdHandle& PresenceSessionId = PresenceSessionsUserMap.FindOrAdd(OpParams.LocalUserId);
+			PresenceSessionId = NewSessionLANRef->SessionId;
+		}
 
-		Op.SetResult(FJoinSession::Result{ NewSessionRef });
+		Op.SetResult(FJoinSession::Result{ });
 
 		TArray<FAccountId> LocalUserIds;
 		LocalUserIds.Reserve(OpParams.LocalUsers.Num());
 		OpParams.LocalUsers.GenerateKeyArray(LocalUserIds);
 
-		FSessionJoined SessionJoinedEvent = { MoveTemp(LocalUserIds), NewSessionRef };
+		FSessionJoined SessionJoinedEvent = { MoveTemp(LocalUserIds), NewSessionLANRef->SessionId };
 		
 		SessionEvents.OnSessionJoined.Broadcast(SessionJoinedEvent);
 	})
@@ -326,17 +349,20 @@ TOnlineAsyncOpHandle<FLeaveSession> FSessionsLAN::LeaveSession(FLeaveSession::Pa
 			return;
 		}
 
-		// This scope guarantees that FoundSession can't be used after its removal
+		TOnlineResult<FGetSessionByName> GetSessionByNameResult = GetSessionByName({ OpParams.SessionName });
+		if (GetSessionByNameResult.IsOk())
 		{
-			TSharedRef<FSession>& FoundSession = LocalSessionsByName.FindChecked(OpParams.SessionName);
+			TSharedRef<const FSession> FoundSession = GetSessionByNameResult.GetOkValue().Session;
+
+			NamedSessionUserMap.FindChecked(OpParams.LocalUserId).Remove(OpParams.SessionName);
 
-			// Only the host should stop the session at this point
-			if (FoundSession->OwnerUserId == OpParams.LocalUserId && FoundSession->SessionSettings.JoinPolicy == ESessionJoinPolicy::Public)
+			if (FoundSession->SessionSettings.bPresenceEnabled)
 			{
-				StopLANSession();
+				PresenceSessionsUserMap.Remove(OpParams.LocalUserId);
 			}
 
-			LocalSessionsByName.Remove(OpParams.SessionName);
+			ClearSessionByName(OpParams.SessionName);
+			ClearSessionById(FoundSession->SessionId);
 		}
 
 		Op.SetResult(FLeaveSession::Result{ });
@@ -386,14 +412,14 @@ TOptional<FOnlineError> FSessionsLAN::TryHostLANSession()
 	return Result;
 }
 
-void FSessionsLAN::FindLANSessions()
+void FSessionsLAN::FindLANSessions(const FAccountId& LocalUserId)
 {
 	// Recreate the unique identifier for this client
 	GenerateNonce((uint8*)&LANSessionManager->LanNonce, 8);
 
 	// Bind delegates
-	FOnValidResponsePacketDelegate ResponseDelegate = FOnValidResponsePacketDelegate::CreateThreadSafeSP(this, &FSessionsLAN::OnValidResponsePacketReceived);
-	FOnSearchingTimeoutDelegate TimeoutDelegate = FOnSearchingTimeoutDelegate::CreateThreadSafeSP(this, &FSessionsLAN::OnLANSearchTimeout);
+	FOnValidResponsePacketDelegate ResponseDelegate = FOnValidResponsePacketDelegate::CreateThreadSafeSP(this, &FSessionsLAN::OnValidResponsePacketReceived, LocalUserId);
+	FOnSearchingTimeoutDelegate TimeoutDelegate = FOnSearchingTimeoutDelegate::CreateThreadSafeSP(this, &FSessionsLAN::OnLANSearchTimeout, LocalUserId);
 
 	FNboSerializeToBuffer Packet(LAN_BEACON_MAX_PACKET_SIZE);
 	LANSessionManager->CreateClientQueryPacket(Packet, LANSessionManager->LanNonce);
@@ -407,10 +433,9 @@ void FSessionsLAN::FindLANSessions()
 		}
 
 		// Trigger the delegate as having failed
-		CurrentSessionSearchHandle->SetError(Errors::RequestFailure()); // TODO: May need a new error type
+		CurrentSessionSearchHandlesUserMap.FindChecked(LocalUserId)->SetError(Errors::RequestFailure()); // TODO: May need a new error type
 
-		CurrentSessionSearch.Reset();
-		CurrentSessionSearchHandle.Reset();
+		CurrentSessionSearchHandlesUserMap.Remove(LocalUserId);
 
 		// If we were hosting public sessions before the search, we'll return the beacon to that state
 		if (PublicSessionsHosted > 0)
@@ -439,7 +464,7 @@ void FSessionsLAN::StopLANSession()
 void FSessionsLAN::OnValidQueryPacketReceived(uint8* PacketData, int32 PacketLength, uint64 ClientNonce)
 {
 	// Iterate through all registered sessions and respond for each one that can be joinable
-	for (const TPair<FName, TSharedRef<FSession>>& Entry : LocalSessionsByName)
+	for (const TPair<FOnlineSessionIdHandle, TSharedRef<FSession>>& Entry : AllSessionsById)
 	{
 		const TSharedRef<FSession>& Session = Entry.Value;
 
@@ -465,50 +490,35 @@ void FSessionsLAN::OnValidQueryPacketReceived(uint8* PacketData, int32 PacketLen
 	}
 }
 
-void FSessionsLAN::OnValidResponsePacketReceived(uint8* PacketData, int32 PacketLength)
+void FSessionsLAN::OnValidResponsePacketReceived(uint8* PacketData, int32 PacketLength, const FAccountId LocalUserId)
 {
-	if (CurrentSessionSearch.IsValid())
+	if (TArray<FOnlineSessionIdHandle>* SearchResults = SearchResultsUserMap.Find(LocalUserId))
 	{
 		TSharedRef<FSessionLAN> Session = MakeShared<FSessionLAN>();
 
 		FNboSerializeFromBuffer Packet(PacketData, PacketLength);
 		ReadSessionFromPacket(Packet, *Session);
 
-		CurrentSessionSearch->FoundSessions.Add(Session);
+		SearchResults->Add(Session->SessionId);
 
-		TMap<FOnlineSessionIdHandle, TSharedRef<FSession>>& UserMap = SessionSearchResultsUserMap.FindOrAdd(CurrentSessionSearchHandle->GetParams().LocalUserId);
-		UserMap.Emplace(Session->SessionId, Session);
+		AllSessionsById.Emplace(Session->SessionId, Session);
 	}
 }
 
-void FSessionsLAN::OnLANSearchTimeout()
+void FSessionsLAN::OnLANSearchTimeout(const FAccountId LocalUserId)
 {
 	if (LANSessionManager->GetBeaconState() == ELanBeaconState::Searching)
 	{
 		LANSessionManager->StopLANSession();
 	}
 
-	if (!CurrentSessionSearch.IsValid())
-	{
-		UE_LOG(LogTemp, Warning, TEXT("[FSessionsLAN::OnLANSearchTimeout] Session search invalid"));
-
-		CurrentSessionSearchHandle->SetError(Errors::InvalidState());
-	}
-
-	if (CurrentSessionSearch->FoundSessions.IsEmpty())
-	{
-		UE_LOG(LogTemp, Warning, TEXT("[FSessionsLAN::OnLANSearchTimeout] Session search found no results"));
-
-		CurrentSessionSearchHandle->SetError(Errors::NotFound());
-	}
+	UE_LOG(LogTemp, Warning, TEXT("[FSessionsLAN::OnLANSearchTimeout] %d sessions found!"), SearchResultsUserMap.FindChecked(LocalUserId).Num());
 
-	UE_LOG(LogTemp, Warning, TEXT("[FSessionsLAN::OnLANSearchTimeout] %d sessions found!"), CurrentSessionSearch->FoundSessions.Num());
+	TArray<FOnlineSessionIdHandle>& FoundSessionIds = SearchResultsUserMap.FindChecked(LocalUserId);
 
-	FFindSessions::Result Result = *CurrentSessionSearch;
-	CurrentSessionSearchHandle->SetResult(MoveTemp(Result));
+	CurrentSessionSearchHandlesUserMap.FindChecked(LocalUserId)->SetResult({ FoundSessionIds });
 
-	CurrentSessionSearch.Reset();
-	CurrentSessionSearchHandle.Reset();
+	CurrentSessionSearchHandlesUserMap.Remove(LocalUserId);
 
 	// If we were hosting public sessions before the search, we'll return the beacon to that state
 	if (PublicSessionsHosted > 0)
diff --git a/Engine/Plugins/Online/Experimental/OnlineServices/Source/OnlineServicesCommonEngineUtils/Public/Online/SessionsLAN.h b/Engine/Plugins/Online/Experimental/OnlineServices/Source/OnlineServicesCommonEngineUtils/Public/Online/SessionsLAN.h
--- a/Engine/Plugins/Online/Experimental/OnlineServices/Source/OnlineServicesCommonEngineUtils/Public/Online/SessionsLAN.h
+++ b/Engine/Plugins/Online/Experimental/OnlineServices/Source/OnlineServicesCommonEngineUtils/Public/Online/SessionsLAN.h
@@ -64,11 +64,11 @@ public:
 
 protected:
 	TOptional<FOnlineError> TryHostLANSession();
-	void FindLANSessions();
+	void FindLANSessions(const FAccountId& LocalUserId);
 	void StopLANSession();
 	void OnValidQueryPacketReceived(uint8* PacketData, int32 PacketLength, uint64 ClientNonce);
-	void OnValidResponsePacketReceived(uint8* PacketData, int32 PacketLength);
-	void OnLANSearchTimeout();
+	void OnValidResponsePacketReceived(uint8* PacketData, int32 PacketLength, const FAccountId LocalUserId);
+	void OnLANSearchTimeout(const FAccountId LocalUserId);
 	virtual void AppendSessionToPacket(FNboSerializeToBuffer& Packet, const FSessionLAN& Session) = 0;
 	virtual void ReadSessionFromPacket(FNboSerializeFromBuffer& Packet, FSessionLAN& Session) = 0;
 
diff --git a/Engine/Plugins/Online/Experimental/OnlineServices/Source/OnlineServicesInterface/Private/Online/Sessions.cpp b/Engine/Plugins/Online/Experimental/OnlineServices/Source/OnlineServicesInterface/Private/Online/Sessions.cpp
--- a/Engine/Plugins/Online/Experimental/OnlineServices/Source/OnlineServicesInterface/Private/Online/Sessions.cpp
+++ b/Engine/Plugins/Online/Experimental/OnlineServices/Source/OnlineServicesInterface/Private/Online/Sessions.cpp
@@ -4,53 +4,6 @@
 
 namespace UE::Online {
 
-const TCHAR* LexToString(ESessionState Value)
-{
-	switch (Value)
-	{
-	case ESessionState::Creating:	return TEXT("Creating");
-	case ESessionState::Joining:	return TEXT("Joining");
-	case ESessionState::Valid:		return TEXT("Valid");
-	case ESessionState::Leaving:	return TEXT("Leaving");
-	case ESessionState::Destroying:	return TEXT("Destroying");
-	default:						checkNoEntry(); // Intentional fallthrough
-	case ESessionState::Invalid:	return TEXT("Invalid");
-	}
-}
-
-void LexFromString(ESessionState& Value, const TCHAR* InStr)
-{
-	if (FCString::Stricmp(InStr, TEXT("Invalid")) == 0)
-	{
-		Value = ESessionState::Invalid;
-	}
-	else if (FCString::Stricmp(InStr, TEXT("Creating")) == 0)
-	{
-		Value = ESessionState::Creating;
-	}
-	else if (FCString::Stricmp(InStr, TEXT("Joining")) == 0)
-	{
-		Value = ESessionState::Joining;
-	}
-	else if (FCString::Stricmp(InStr, TEXT("Valid")) == 0)
-	{
-		Value = ESessionState::Valid;
-	}
-	else if (FCString::Stricmp(InStr, TEXT("Leaving")) == 0)
-	{
-		Value = ESessionState::Leaving;
-	}
-	else if (FCString::Stricmp(InStr, TEXT("Destroying")) == 0)
-	{
-		Value = ESessionState::Destroying;
-	}
-	else
-	{
-		checkNoEntry();
-		Value = ESessionState::Invalid;
-	}
-}
-
 const TCHAR* LexToString(ESessionJoinPolicy Value)
 {
 	switch (Value)
@@ -242,7 +195,6 @@ FSession::FSession()
 FSession::FSession(const FSession& InSession)
 	: OwnerUserId(InSession.OwnerUserId)
 	, SessionId(InSession.SessionId)
-	, CurrentState(InSession.CurrentState)
 	, SessionSettings(InSession.SessionSettings)
 {
 
diff --git a/Engine/Plugins/Online/Experimental/OnlineServices/Source/OnlineServicesInterface/Public/Online/Sessions.h b/Engine/Plugins/Online/Experimental/OnlineServices/Source/OnlineServicesInterface/Public/Online/Sessions.h
--- a/Engine/Plugins/Online/Experimental/OnlineServices/Source/OnlineServicesInterface/Public/Online/Sessions.h
+++ b/Engine/Plugins/Online/Experimental/OnlineServices/Source/OnlineServicesInterface/Public/Online/Sessions.h
@@ -54,18 +54,6 @@ struct FSessionMemberUpdate
 
 using FSessionMemberUpdatesMap = TMap<FAccountId, FSessionMemberUpdate>;
 
-enum class ESessionState : uint8
-{
-	Invalid,
-	Creating,
-	Joining,
-	Valid,
-	Leaving,
-	Destroying
-};
-ONLINESERVICESINTERFACE_API const TCHAR* LexToString(ESessionState Value);
-ONLINESERVICESINTERFACE_API void LexFromString(ESessionState& Value, const TCHAR* InStr);
-
 enum class ESessionJoinPolicy : uint8
 {
 	Public,
@@ -201,9 +189,6 @@ struct ONLINESERVICESINTERFACE_API FSession
 	/** The id for the session, platform dependent */
 	FOnlineSessionIdHandle SessionId;
 
-	/** The current state of the session */
-	ESessionState CurrentState = ESessionState::Invalid;
-
 	/** Set of session properties that can be altered by the session owner */
 	FSessionSettings SessionSettings;
 };
@@ -220,7 +205,7 @@ struct FSessionInvite
 	FOnlineSessionInviteIdHandle InviteId;
 
 	/* Pointer to the session information */
-	TSharedRef<const FSession> Session;
+	FOnlineSessionIdHandle SessionId;
 
 	// TODO: Default constructor will be deleted after we cache invites
 };
@@ -297,9 +282,7 @@ struct FCreateSession
 
 	struct Result
 	{
-		TSharedRef<const FSession> Session;
 
-		Result() = delete; // cannot default construct due to TSharedRef
 	};
 };
 
@@ -321,9 +304,7 @@ struct FUpdateSession
 
 	struct Result
 	{
-		TSharedRef<const FSession> Session;
 
-		Result() = delete; // cannot default construct due to TSharedRef
 	};
 };
 
@@ -379,7 +360,7 @@ struct FFindSessions
 
 	struct Result
 	{
-		TArray<TSharedRef<const FSession>> FoundSessions;
+		TArray<FOnlineSessionIdHandle> FoundSessionIds;
 	};
 };
 
@@ -407,9 +388,7 @@ struct FStartMatchmaking
 
 	struct Result
 	{
-		TSharedRef<const FSession> Session;
 
-		Result() = delete; // cannot default construct due to TSharedRef
 	};
 };
 
@@ -434,9 +413,7 @@ struct FJoinSession
 
 	struct Result
 	{
-		TSharedRef<const FSession> Session;
 
-		Result() = delete; // cannot default construct due to TSharedRef
 	};
 };
 
@@ -553,6 +530,9 @@ struct FRegisterPlayers
 
 	struct Params
 	{
+		/* The local user agent */
+		FAccountId LocalUserId;
+
 		/* The local name for the session. */
 		FName SessionName;
 
@@ -574,6 +554,9 @@ struct FUnregisterPlayers
 
 	struct Params
 	{
+		/* The local user agent */
+		FAccountId LocalUserId;
+
 		/* The local name for the session. */
 		FName SessionName;
 
@@ -597,9 +580,7 @@ struct FSessionJoined
 	TArray<FAccountId> LocalUserIds;
 
 	/* A shared reference to the session joined. */
-	TSharedRef<const FSession> Session;
-
-	FSessionJoined() = delete; // cannot default construct due to TSharedRef
+	FOnlineSessionIdHandle SessionId;
 };
 
 struct FSessionLeft
@@ -610,13 +591,11 @@ struct FSessionLeft
 
 struct FSessionUpdated
 {
-	/* Reference to the updated session object */
-	TSharedRef<const FSession> UpdatedSession;
+	/* Local name for the updated session object */
+	FName SessionName;
 
 	/* Updated session settings */
 	FSessionSettingsUpdate UpdatedSettings;
-
-	FSessionUpdated() = delete; // cannot default construct due to TSharedRef
 };
 
 struct FSessionInviteReceived
@@ -627,7 +606,7 @@ struct FSessionInviteReceived
 	/** The session invite the local user was sent, or the online error if there was a failure retrieving the session for it*/
 	TSharedRef<const FSessionInvite> SessionInvite;
 
-	FSessionInviteReceived() = delete; // cannot default construct due to TSharedRef
+	FSessionInviteReceived() = delete; // cannot default construct due to TSharedRef. TODO: Add a GetSessionInvite method and return an id here
 };
 
 /** Session join requested source */
@@ -646,13 +625,13 @@ struct FUISessionJoinRequested
 	/** The local user associated with the join request. */
 	FAccountId LocalUserId;
 
-	/** The session the local user requested to join, or the online error if there was a failure retrieving it */
-	TResult<TSharedRef<const FSession>, FOnlineError> Result;
+	/** The id for the session the local user requested to join, or the online error if there was a failure retrieving it */
+	TResult<FOnlineSessionIdHandle, FOnlineError> Result;
 
 	/** Join request source */
 	EUISessionJoinRequestedSource JoinRequestedSource = EUISessionJoinRequestedSource::Unspecified;
 
-	FUISessionJoinRequested() = delete; // cannot default construct due to TSharedRef
+	FUISessionJoinRequested() = delete; // cannot default construct due to TResult
 };
 
 class ISessions
@@ -910,14 +889,14 @@ END_ONLINE_STRUCT_META()
 BEGIN_ONLINE_STRUCT_META(FSession)
 	ONLINE_STRUCT_FIELD(FSession, OwnerUserId),
 	ONLINE_STRUCT_FIELD(FSession, SessionId),
-	ONLINE_STRUCT_FIELD(FSession, CurrentState),
 	ONLINE_STRUCT_FIELD(FSession, SessionSettings)
 END_ONLINE_STRUCT_META()
 
 BEGIN_ONLINE_STRUCT_META(FSessionInvite)
+	ONLINE_STRUCT_FIELD(FSessionInvite, RecipientId),
 	ONLINE_STRUCT_FIELD(FSessionInvite, SenderId),
 	ONLINE_STRUCT_FIELD(FSessionInvite, InviteId),
-	ONLINE_STRUCT_FIELD(FSessionInvite, Session)
+	ONLINE_STRUCT_FIELD(FSessionInvite, SessionId)
 END_ONLINE_STRUCT_META()
 
 BEGIN_ONLINE_STRUCT_META(FGetAllSessions::Params)
@@ -929,7 +908,7 @@ END_ONLINE_STRUCT_META()
 
 BEGIN_ONLINE_STRUCT_META(FGetSessionByName::Params)
 	ONLINE_STRUCT_FIELD(FGetSessionByName::Params, LocalName)
-	END_ONLINE_STRUCT_META()
+END_ONLINE_STRUCT_META()
 
 BEGIN_ONLINE_STRUCT_META(FGetSessionByName::Result)
 	ONLINE_STRUCT_FIELD(FGetSessionByName::Result, Session)
@@ -952,7 +931,6 @@ BEGIN_ONLINE_STRUCT_META(FCreateSession::Params)
 END_ONLINE_STRUCT_META()
 
 BEGIN_ONLINE_STRUCT_META(FCreateSession::Result)
-	ONLINE_STRUCT_FIELD(FCreateSession::Result, Session)
 END_ONLINE_STRUCT_META()
 
 BEGIN_ONLINE_STRUCT_META(FUpdateSession::Params)
@@ -962,7 +940,6 @@ BEGIN_ONLINE_STRUCT_META(FUpdateSession::Params)
 END_ONLINE_STRUCT_META()
 
 BEGIN_ONLINE_STRUCT_META(FUpdateSession::Result)
-	ONLINE_STRUCT_FIELD(FUpdateSession::Result, Session)
 END_ONLINE_STRUCT_META()
 
 BEGIN_ONLINE_STRUCT_META(FLeaveSession::Params)
@@ -985,7 +962,7 @@ BEGIN_ONLINE_STRUCT_META(FFindSessions::Params)
 END_ONLINE_STRUCT_META()
 
 BEGIN_ONLINE_STRUCT_META(FFindSessions::Result)
-	ONLINE_STRUCT_FIELD(FFindSessions::Result, FoundSessions)
+	ONLINE_STRUCT_FIELD(FFindSessions::Result, FoundSessionIds)
 END_ONLINE_STRUCT_META()
 
 BEGIN_ONLINE_STRUCT_META(FStartMatchmaking::Params)
@@ -997,7 +974,6 @@ BEGIN_ONLINE_STRUCT_META(FStartMatchmaking::Params)
 END_ONLINE_STRUCT_META()
 
 BEGIN_ONLINE_STRUCT_META(FStartMatchmaking::Result)
-	ONLINE_STRUCT_FIELD(FStartMatchmaking::Result, Session)
 END_ONLINE_STRUCT_META()
 
 BEGIN_ONLINE_STRUCT_META(FJoinSession::Params)
@@ -1008,7 +984,6 @@ BEGIN_ONLINE_STRUCT_META(FJoinSession::Params)
 END_ONLINE_STRUCT_META()
 
 BEGIN_ONLINE_STRUCT_META(FJoinSession::Result)
-	ONLINE_STRUCT_FIELD(FJoinSession::Result, Session)
 END_ONLINE_STRUCT_META()
 
 BEGIN_ONLINE_STRUCT_META(FAddSessionMembers::Params)
@@ -1057,6 +1032,7 @@ BEGIN_ONLINE_STRUCT_META(FRejectSessionInvite::Result)
 END_ONLINE_STRUCT_META()
 
 BEGIN_ONLINE_STRUCT_META(FRegisterPlayers::Params)
+	ONLINE_STRUCT_FIELD(FRegisterPlayers::Params, LocalUserId),
 	ONLINE_STRUCT_FIELD(FRegisterPlayers::Params, SessionName),
 	ONLINE_STRUCT_FIELD(FRegisterPlayers::Params, TargetUsers),
 	ONLINE_STRUCT_FIELD(FRegisterPlayers::Params, bReserveSlot)
@@ -1066,6 +1042,7 @@ BEGIN_ONLINE_STRUCT_META(FRegisterPlayers::Result)
 END_ONLINE_STRUCT_META()
 
 BEGIN_ONLINE_STRUCT_META(FUnregisterPlayers::Params)
+	ONLINE_STRUCT_FIELD(FUnregisterPlayers::Params, LocalUserId),
 	ONLINE_STRUCT_FIELD(FUnregisterPlayers::Params, SessionName),
 	ONLINE_STRUCT_FIELD(FUnregisterPlayers::Params, TargetUsers),
 	ONLINE_STRUCT_FIELD(FUnregisterPlayers::Params, bRemoveUnregisteredPlayers)
diff --git a/Engine/Plugins/Online/Experimental/OnlineServicesEOSGS/Source/Private/Online/SessionsEOSGS.cpp b/Engine/Plugins/Online/Experimental/OnlineServicesEOSGS/Source/Private/Online/SessionsEOSGS.cpp
--- a/Engine/Plugins/Online/Experimental/OnlineServicesEOSGS/Source/Private/Online/SessionsEOSGS.cpp
+++ b/Engine/Plugins/Online/Experimental/OnlineServicesEOSGS/Source/Private/Online/SessionsEOSGS.cpp
@@ -53,7 +53,6 @@ FSessionEOSGS::FSessionEOSGS(const EOS_HSessionDetails& InSessionDetailsHandle)
 	EOS_EResult CopyInfoResult = EOS_SessionDetails_CopyInfo(InSessionDetailsHandle, &CopyInfoOptions, &SessionDetailsInfo);
 	if (CopyInfoResult == EOS_EResult::EOS_Success)
 	{
-		CurrentState = ESessionState::Valid;
 		SessionId = FSessionsEOSGS::CreateSessionId(FString(SessionDetailsInfo->SessionId));
 		SessionSettings.NumOpenPublicConnections = SessionDetailsInfo->NumOpenPublicConnections;
 		// We could retrieve the Host Address here if we set it during creation or update
@@ -195,8 +194,6 @@ FSessionEOSGS::FSessionEOSGS(const EOS_HSessionDetails& InSessionDetailsHandle)
 	else
 	{
 		UE_LOG(LogTemp, Warning, TEXT("[FSessionEOSGS] EOS_SessionDetails_CopyInfo failed with result [%s]"), *LexToString(CopyInfoResult));
-
-		CurrentState = ESessionState::Invalid;
 	}
 }
 
@@ -279,38 +276,41 @@ void FSessionsEOSGS::HandleSessionInviteReceived(const EOS_Sessions_SessionInvit
 	FAccountId LocalUserId = FindAccountId(Data->LocalUserId);
 	if (LocalUserId.IsValid())
 	{
-		Services.Get<FAuthEOSGS>()->ResolveAccountIds(LocalUserId, { Data->LocalUserId, Data->TargetUserId })
-			.Next([this, WeakThis = AsWeak(), LocalUserId, InviteId = FString(UTF8_TO_TCHAR(Data->InviteId))](TArray<FAccountId> ResolvedAccountIds)
+		FString InviteId(UTF8_TO_TCHAR(Data->InviteId));
+
+		BuildSessionFromInvite(LocalUserId, InviteId)
+			.OnComplete([this, WeakThis = AsWeak(), SenderPUID = Data->TargetUserId, InviteId](const TOnlineResult<FBuildSessionFromDetailsHandle>& Result)
 		{
-			if (const TSharedPtr<ISessions> StrongThis = WeakThis.Pin())
+			if (Result.IsOk())
 			{
-				// First and second place in the array will be occupied by the receiver and the sender, respectively, since the same order is kept in the array of resolved ids
-				const FAccountId& ReceiverId = ResolvedAccountIds[0];
-				const FAccountId& SenderId = ResolvedAccountIds[1];
-
-				TResult<TSharedRef<const FSession>, FOnlineError> SessionInviteResult = BuildSessionFromInvite(InviteId);
-				
-				if (SessionInviteResult.IsOk())
+				if (const TSharedPtr<ISessions> StrongThis = WeakThis.Pin())
 				{
+					// First and second place in the array will be occupied by the receiver and the sender, respectively, since the same order is kept in the array of resolved ids
+					const FAccountId& ReceiverId = Result.GetOkValue().LocalUserId;
+					const FAccountId& SenderId = FindAccountId(SenderPUID);
+
 					FOnlineSessionInviteIdHandle InviteIdHandle = CreateSessionInviteId(InviteId);
 
+					TSharedRef<FSession> Session = Result.GetOkValue().Session;
+
 					TSharedRef<FSessionInvite> SessionInviteRef = MakeShared<FSessionInvite>(FSessionInvite{
 						ReceiverId,
 						SenderId,
 						InviteIdHandle,
-						SessionInviteResult.GetOkValue()
+						Session->SessionId
 					});
 
-					TMap<FOnlineSessionInviteIdHandle, TSharedRef<FSessionInvite>>& UserMap = SessionInvitesUserMap.FindOrAdd(ReceiverId);
-					UserMap.Emplace(InviteIdHandle, SessionInviteRef);
+					SessionInvitesUserMap.FindOrAdd(ReceiverId).Emplace(InviteIdHandle, SessionInviteRef);
+
+					AllSessionsById.Emplace(Session->SessionId, Session);
 
-					FSessionInviteReceived Event = { LocalUserId, SessionInviteRef };
+					FSessionInviteReceived Event = { ReceiverId, SessionInviteRef };
 
 					SessionEvents.OnSessionInviteReceived.Broadcast(Event);
 				}
-
-				// We won't broadcast the event if there was an error retrieving the session information
 			}
+
+			// We won't broadcast the event if there was an error retrieving the session information
 		});
 	}
 }
@@ -320,15 +320,30 @@ void FSessionsEOSGS::HandleSessionInviteAccepted(const EOS_Sessions_SessionInvit
 	FAccountId LocalUserId = FindAccountId(Data->LocalUserId);
 	if (LocalUserId.IsValid())
 	{
-		FUISessionJoinRequested Event{
-			LocalUserId,
-			BuildSessionFromInvite(FString(Data->InviteId)),
-			EUISessionJoinRequestedSource::FromInvitation
-		};
+		FString InviteId(UTF8_TO_TCHAR(Data->InviteId));
 
-		SessionEvents.OnUISessionJoinRequested.Broadcast(Event);
+		BuildSessionFromInvite(LocalUserId, InviteId)
+			.OnComplete([this, WeakThis = AsWeak()](const TOnlineResult<FBuildSessionFromDetailsHandle>& Result)
+		{
+			if (Result.IsOk())
+			{
+				if (const TSharedPtr<ISessions> StrongThis = WeakThis.Pin())
+				{
+					// Instead of using the session information we have stored, we'll use this opportunity to get updated data for the session
+					AllSessionsById.Emplace(Result.GetOkValue().Session->SessionId, Result.GetOkValue().Session);
+
+					FUISessionJoinRequested Event{
+						Result.GetOkValue().LocalUserId,
+						TResult<FOnlineSessionIdHandle, FOnlineError>(Result.GetOkValue().Session->SessionId),
+						EUISessionJoinRequestedSource::FromInvitation
+					};
+
+					SessionEvents.OnUISessionJoinRequested.Broadcast(Event);
 
-		// The game can react to the OnUISessionJoinRequested event by starting the JoinSession process
+					// The game can react to the OnUISessionJoinRequested event by starting the JoinSession process
+				}
+			}
+		});
 	}
 }
 
@@ -337,15 +352,28 @@ void FSessionsEOSGS::HandleJoinSessionAccepted(const EOS_Sessions_JoinSessionAcc
 	FAccountId LocalUserId = FindAccountId(Data->LocalUserId);
 	if (LocalUserId.IsValid())
 	{
-		FUISessionJoinRequested Event {
-			LocalUserId,
-			BuildSessionFromUIEvent(Data->UiEventId),
-			EUISessionJoinRequestedSource::Unspecified
-		};
+		BuildSessionFromUIEvent(LocalUserId, Data->UiEventId)
+			.OnComplete([this, WeakThis = AsWeak()](const TOnlineResult<FBuildSessionFromDetailsHandle>& Result)
+		{
+			if (Result.IsOk())
+			{
+				if (const TSharedPtr<ISessions> StrongThis = WeakThis.Pin())
+				{
+					// Instead of using the session information we have stored, we'll use this opportunity to get updated data for the session
+					AllSessionsById.Emplace(Result.GetOkValue().Session->SessionId, Result.GetOkValue().Session);
 
-		SessionEvents.OnUISessionJoinRequested.Broadcast(Event);
+					FUISessionJoinRequested Event{
+						Result.GetOkValue().LocalUserId,
+						TResult<FOnlineSessionIdHandle, FOnlineError>(Result.GetOkValue().Session->SessionId),
+						EUISessionJoinRequestedSource::Unspecified
+					};
 
-		// The game can react to the OnUISessionJoinRequested event by starting the JoinSession process
+					SessionEvents.OnUISessionJoinRequested.Broadcast(Event);
+
+					// The game can react to the OnUISessionJoinRequested event by starting the JoinSession process
+				}
+			}
+		});
 	}
 }
 
@@ -450,14 +478,6 @@ TOnlineAsyncOpHandle<FCreateSession> FSessionsEOSGS::CreateSession(FCreateSessio
 		// We write all SessionSettings values into the SessionModificationHandle
 		WriteCreateSessionModificationHandle(SessionModificationHandle, OpParams.SessionSettings);
 
-		// If the session modification handle is created successfully, we'll create the local session object
- 		TSharedRef<FSessionEOSGS> NewSessionEOSGSRef = MakeShared<FSessionEOSGS>();
-		NewSessionEOSGSRef->CurrentState = ESessionState::Creating;
-		NewSessionEOSGSRef->OwnerUserId = OpParams.LocalUserId;
-		NewSessionEOSGSRef->SessionSettings = OpParams.SessionSettings;
-
-		LocalSessionsByName.Add(OpParams.SessionName, NewSessionEOSGSRef);
-
 		// Always update joinability on session creation
 		FUpdateSessionImpl::Params UpdateSessionImplParams{ MakeShared<FSessionModificationHandleEOSGS>(SessionModificationHandle), FUpdateSessionJoinabilityParams{ OpParams.SessionName, OpParams.SessionSettings.bAllowNewMembers } };
 		
@@ -469,18 +489,28 @@ TOnlineAsyncOpHandle<FCreateSession> FSessionsEOSGS::CreateSession(FCreateSessio
 
 		if (Result.IsOk())
 		{
-			TSharedRef<FSession>& FoundSession = LocalSessionsByName.FindChecked(OpParams.SessionName);
-			FoundSession->CurrentState = ESessionState::Valid;
+			// If the session is created successfully, we'll create the local session object
+			TSharedRef<FSessionEOSGS> NewSessionEOSGSRef = MakeShared<FSessionEOSGS>();
+			NewSessionEOSGSRef->OwnerUserId = OpParams.LocalUserId;
+			NewSessionEOSGSRef->SessionSettings = OpParams.SessionSettings;
+			NewSessionEOSGSRef->SessionId = CreateSessionId(Result.GetOkValue().NewSessionId);
 
-			// TODO: Make additional local users join the created session
+			AllSessionsById.Emplace(NewSessionEOSGSRef->SessionId, NewSessionEOSGSRef);
+
+			LocalSessionsByName.Emplace(OpParams.SessionName, NewSessionEOSGSRef->SessionId);
+
+			NamedSessionUserMap.FindOrAdd(OpParams.LocalUserId).AddUnique(OpParams.SessionName);
+
+			if (NewSessionEOSGSRef->SessionSettings.bPresenceEnabled)
+			{
+				FOnlineSessionIdHandle& PresenceSessionId = PresenceSessionsUserMap.FindOrAdd(OpParams.LocalUserId);
+				PresenceSessionId = NewSessionEOSGSRef->SessionId;
+			}
 
-			Op.SetResult(FCreateSession::Result{ FoundSession });
+			Op.SetResult({ });
 		}
 		else
 		{
-			// If creation failed, we remove the session objects
-			LocalSessionsByName.Remove(OpParams.SessionName);
-
 			Op.SetError(MoveTemp(Result.GetErrorValue()));
 		}
 	})
@@ -706,7 +736,10 @@ void FSessionsEOSGS::WriteUpdateSessionModificationHandle(EOS_HSessionModificati
 		AddAttribute(SessionModificationHandle, EOSGS_SESSION_ID_OVERRIDE, { FSchemaVariant(NewSettings.SessionIdOverride.GetValue()), ESchemaAttributeVisibility::Public });
 	}
 
-	const TSharedRef<FSession>& FoundSession = LocalSessionsByName.FindChecked(SessionName);
+	TOnlineResult<FGetSessionByName> GetSessionByNameResult = GetSessionByName({ SessionName });
+	check(GetSessionByNameResult.IsOk());
+
+	TSharedRef<const FSession> FoundSession = GetSessionByNameResult.GetOkValue().Session;
 	const FSessionSettings& SessionSettings = FoundSession->SessionSettings;
 
 	if (NewSettings.JoinPolicy.IsSet())
@@ -861,20 +894,23 @@ TOnlineAsyncOpHandle<FUpdateSession> FSessionsEOSGS::UpdateSession(FUpdateSessio
 		{
 			const FUpdateSession::Params& OpParams = Op.GetParams();
 
-			if (TSharedRef<FSession>* Session = LocalSessionsByName.Find(OpParams.SessionName))
+			TOnlineResult<FGetMutableSessionByName> GetMutableSessionByNameResult = GetMutableSessionByName({ OpParams.SessionName });
+			if (GetMutableSessionByNameResult.IsOk())
 			{
+				TSharedRef<FSession> FoundSession = GetMutableSessionByNameResult.GetOkValue().Session;
+
 				// Now that the API Session update has processed successfully, we'll update our local session with the same data
-				Session->Get().SessionSettings += OpParams.Mutations;
+				FoundSession->SessionSettings += OpParams.Mutations;
 
 				// We set the result and fire the event
-				Op.SetResult(FUpdateSession::Result{ *Session });
+				Op.SetResult({ });
 
-				FSessionUpdated SessionUpdatedEvent{ *Session, OpParams.Mutations };
+				FSessionUpdated SessionUpdatedEvent{ OpParams.SessionName, OpParams.Mutations };
 				SessionEvents.OnSessionUpdated.Broadcast(SessionUpdatedEvent);
 			}
 			else
 			{
-				Op.SetError(Errors::NotFound());
+				Op.SetError(MoveTemp(GetMutableSessionByNameResult.GetErrorValue()));
 			}
 		}
 		else
@@ -909,22 +945,15 @@ TFuture<TDefaultErrorResult<FUpdateSessionImpl>> FSessionsEOSGS::UpdateSessionIm
 			return;
 		}
 
-		// In session creation calls, we'll need to set the new id for the session
-		TSharedRef<FSession>& FoundSession = LocalSessionsByName.FindChecked(FName(UTF8_TO_TCHAR(Result->SessionName)));
-		if (!FoundSession->SessionId.IsValid())
-		{
-			FoundSession->SessionId = CreateSessionId(UTF8_TO_TCHAR(Result->SessionId));
-		}
-
 		// After the successful general update, if indicated, we'll update the joinability
 		if (Params.UpdateJoinabilitySettings.IsSet())
 		{
 			UpdateSessionJoinabilityImpl(MoveTemp(Params.UpdateJoinabilitySettings.GetValue()))
-				.Next([this, Promise = MoveTemp(Promise)](TDefaultErrorResult<FUpdateSessionJoinabilityImpl>&& Result) mutable
+				.Next([this, NewSessionId = UTF8_TO_TCHAR(Result->SessionId), Promise = MoveTemp(Promise)](TDefaultErrorResult<FUpdateSessionJoinabilityImpl>&& Result) mutable
 			{
 				if (Result.IsOk())
 				{
-					Promise.EmplaceValue(FUpdateSessionImpl::Result{ });
+					Promise.EmplaceValue(FUpdateSessionImpl::Result{ NewSessionId });
 				}
 				else
 				{
@@ -934,7 +963,7 @@ TFuture<TDefaultErrorResult<FUpdateSessionImpl>> FSessionsEOSGS::UpdateSessionIm
 		}
 		else
 		{
-			Promise.EmplaceValue(FUpdateSessionImpl::Result{ });
+			Promise.EmplaceValue(FUpdateSessionImpl::Result{ FString(Result->SessionId) });
 		}
 	});
 
@@ -1115,7 +1144,21 @@ TOnlineAsyncOpHandle<FLeaveSession> FSessionsEOSGS::LeaveSession(FLeaveSession::
 
 		const FLeaveSession::Params& OpParams = Op.GetParams();
 
-		LocalSessionsByName.Remove(OpParams.SessionName);
+		TOnlineResult<FGetSessionByName> GetSessionByNameResult = GetSessionByName({ OpParams.SessionName });
+		if (GetSessionByNameResult.IsOk())
+		{
+			TSharedRef<const FSession> FoundSession = GetSessionByNameResult.GetOkValue().Session;
+
+			NamedSessionUserMap.FindChecked(OpParams.LocalUserId).Remove(OpParams.SessionName);
+
+			if (FoundSession->SessionSettings.bPresenceEnabled)
+			{
+				PresenceSessionsUserMap.Remove(OpParams.LocalUserId);
+			}
+
+			ClearSessionByName(OpParams.SessionName);
+			ClearSessionById(FoundSession->SessionId);
+		}
 
 		Op.SetResult(FLeaveSession::Result{ });
 
@@ -1123,7 +1166,6 @@ TOnlineAsyncOpHandle<FLeaveSession> FSessionsEOSGS::LeaveSession(FLeaveSession::
 		SessionLeftEvent.LocalUserIds.Add(OpParams.LocalUserId);
 		//SessionLeftEvent.LocalUserIds.Append(OpParams.AdditionalLocalUsers); // TODO: Process multiple users in Session destruction
 		SessionEvents.OnSessionLeft.Broadcast(SessionLeftEvent);
-
 	})
 	.Enqueue(GetSerialQueue());
 
@@ -1260,11 +1302,9 @@ TOnlineAsyncOpHandle<FFindSessions> FSessionsEOSGS::FindSessions(FFindSessions::
 			return;
 		}
 
-		// Before we start the search, we reset the variables and cache
-		CurrentSessionSearch = MakeShared<FFindSessions::Result>();
-		CurrentSessionSearchHandle = Op.AsShared();
-
-		SessionSearchResultsUserMap.FindOrAdd(OpParams.LocalUserId).Reset();
+		// Before we start the search, we reset the cache
+		SearchResultsUserMap.FindOrAdd(OpParams.LocalUserId).Reset();
+		CurrentSessionSearchHandlesUserMap.Emplace(OpParams.LocalUserId, Op.AsShared());
 
 		// We start preparing the search
 		EOS_Sessions_CreateSessionSearchOptions CreateSessionSearchOptions = { };
@@ -1283,7 +1323,7 @@ TOnlineAsyncOpHandle<FFindSessions> FSessionsEOSGS::FindSessions(FFindSessions::
 			return;
 		}
 
-		CurrentSessionSearchHandleEOSGS = MakeShared<FSessionSearchHandleEOSGS>(SearchHandle);
+		const TSharedRef<FSessionSearchHandleEOSGS>& CurrentSessionSearchHandleEOSGS = CurrentSessionSearchHandleEOSGSUserMap.Emplace(OpParams.LocalUserId, MakeShared<FSessionSearchHandleEOSGS>(SearchHandle));
 
 		// We write the search attributes
 		WriteSessionSearchHandle(*CurrentSessionSearchHandleEOSGS, OpParams);
@@ -1305,6 +1345,8 @@ TOnlineAsyncOpHandle<FFindSessions> FSessionsEOSGS::FindSessions(FFindSessions::
 
 		const FFindSessions::Params& OpParams = Op.GetParams();
 
+		const TSharedRef<FSessionSearchHandleEOSGS>& CurrentSessionSearchHandleEOSGS = CurrentSessionSearchHandleEOSGSUserMap.FindChecked(OpParams.LocalUserId);
+
 		// For a successful session, we'll get the search results
 		EOS_SessionSearch_GetSearchResultCountOptions GetSearchResultCountOptions = { };
 		GetSearchResultCountOptions.ApiVersion = EOS_SESSIONSEARCH_GETSEARCHRESULTCOUNT_API_LATEST;
@@ -1328,16 +1370,21 @@ TOnlineAsyncOpHandle<FFindSessions> FSessionsEOSGS::FindSessions(FFindSessions::
 			{
 				TSharedRef<TPromise<TDefaultErrorResult<FBuildSessionFromDetailsHandle>>> BuildSessionPromise = MakeShared<TPromise<TDefaultErrorResult<FBuildSessionFromDetailsHandle>>>();
 				BuildSessionFromDetailsHandle({ OpParams.LocalUserId, MakeShared<FSessionDetailsHandleEOSGS>(SessionDetailsHandle) })
-				.OnComplete([this, BuildSessionPromise](const TOnlineResult<FBuildSessionFromDetailsHandle>& Result) mutable
+					.OnComplete([this, WeakOp = Op.AsWeak(), BuildSessionPromise](const TOnlineResult<FBuildSessionFromDetailsHandle>& Result) mutable
 				{
 					if (Result.IsOk())
 					{
-						CurrentSessionSearch->FoundSessions.Add(Result.GetOkValue().Session);
+						if (TOnlineAsyncOpPtr<FFindSessions> StrongOp = WeakOp.Pin())
+						{
+							const TSharedRef<FSession>& Session = Result.GetOkValue().Session;
+
+							TArray<FOnlineSessionIdHandle>& SearchResults = SearchResultsUserMap.FindOrAdd(Result.GetOkValue().LocalUserId);
+							SearchResults.Add(Session->SessionId);
 
-						TMap<FOnlineSessionIdHandle, TSharedRef<FSession>>& UserMap = SessionSearchResultsUserMap.FindOrAdd(Result.GetOkValue().LocalUserId);
-						UserMap.Emplace(Result.GetOkValue().Session->SessionId, Result.GetOkValue().Session);
+							AllSessionsById.Emplace(Session->SessionId, Session);
 
-						BuildSessionPromise->EmplaceValue(Result.GetOkValue());
+							BuildSessionPromise->EmplaceValue(Result.GetOkValue());
+						}
 					}
 					else
 					{
@@ -1370,10 +1417,14 @@ TOnlineAsyncOpHandle<FFindSessions> FSessionsEOSGS::FindSessions(FFindSessions::
 					}
 				}
 
-				StrongOp->SetResult(MoveTemp(*CurrentSessionSearch));
+				const FFindSessions::Params& OpParams = StrongOp->GetParams();
+
+				TArray<FOnlineSessionIdHandle>& SearchResults = SearchResultsUserMap.FindChecked(OpParams.LocalUserId);
+
+				StrongOp->SetResult({ SearchResults });
 
-				CurrentSessionSearch.Reset();
-				CurrentSessionSearchHandleEOSGS.Reset();
+				CurrentSessionSearchHandlesUserMap.Remove(OpParams.LocalUserId);
+				CurrentSessionSearchHandleEOSGSUserMap.Remove(OpParams.LocalUserId);
 			}
 		});
 	});
@@ -1458,10 +1509,6 @@ TOnlineAsyncOpHandle<FJoinSession> FSessionsEOSGS::JoinSession(FJoinSession::Par
 		}
 
 		const TSharedRef<const FSession>& FoundSession = GetSessionByIdResult.GetOkValue().Session;
-		TSharedRef<FSession> NewSessionRef = ConstCastSharedRef<FSession>(FoundSession);
-		NewSessionRef->CurrentState = ESessionState::Joining;
-
-		LocalSessionsByName.Add(OpParams.SessionName, NewSessionRef);
 
 		// We start setup for the API call
 		EOS_Sessions_JoinSessionOptions JoinSessionOptions = { };
@@ -1486,50 +1533,66 @@ TOnlineAsyncOpHandle<FJoinSession> FSessionsEOSGS::JoinSession(FJoinSession::Par
 
 		if (Result->ResultCode != EOS_EResult::EOS_Success)
 		{
-			LocalSessionsByName.Remove(OpParams.SessionName);
-
 			Op.SetError(Errors::FromEOSResult(Result->ResultCode));
 			return;
 		}
 		
-		if (TSharedRef<FSession>* FoundSession = LocalSessionsByName.Find(OpParams.SessionName))
+		TOnlineResult<FGetSessionById> GetSessionByIdResult = GetSessionById({ OpParams.LocalUserId, OpParams.SessionId });
+		if (GetSessionByIdResult.IsError())
+		{
+			// If no result is found, the id might be expired, which we should notify
+			if (FOnlineSessionIdRegistryEOSGS::Get().IsSessionIdExpired(OpParams.SessionId))
+			{
+				UE_LOG(LogTemp, Warning, TEXT("[FSessionsEOSGS::JoinSession] SessionId parameter [%s] is expired. Please call FindSessions to get an updated list of available sessions "), *ToLogString(OpParams.SessionId));
+			}
+
+			Op.SetError(MoveTemp(GetSessionByIdResult.GetErrorValue()));
+			return;
+		}
+
+		const TSharedRef<const FSession>& FoundSession = GetSessionByIdResult.GetOkValue().Session;
+
+		LocalSessionsByName.Emplace(OpParams.SessionName, OpParams.SessionId);
+
+		NamedSessionUserMap.FindOrAdd(OpParams.LocalUserId).AddUnique(OpParams.SessionName);
+
+		if (FoundSession->SessionSettings.bPresenceEnabled)
 		{
-			FoundSession->Get().CurrentState = ESessionState::Valid;
+			FOnlineSessionIdHandle& PresenceSessionId = PresenceSessionsUserMap.FindOrAdd(OpParams.LocalUserId);
+			PresenceSessionId = FoundSession->SessionId;
+		}
 
-			// After successfully joining a session, we'll remove all related invites if any are found
-			if (TMap<FOnlineSessionInviteIdHandle, TSharedRef<FSessionInvite>>* UserMap = SessionInvitesUserMap.Find(OpParams.LocalUserId))
+		// After successfully joining a session, we'll remove all related invites if any are found
+		if (TMap<FOnlineSessionInviteIdHandle, TSharedRef<FSessionInvite>>* UserMap = SessionInvitesUserMap.Find(OpParams.LocalUserId))
+		{
+			TArray<FOnlineSessionInviteIdHandle> InviteIdsToRemove;
+			for (const TPair<FOnlineSessionInviteIdHandle, TSharedRef<FSessionInvite>>& Entry : *UserMap)
 			{
-				TArray<FOnlineSessionInviteIdHandle> InviteIdsToRemove;
-				for (const TPair<FOnlineSessionInviteIdHandle, TSharedRef<FSessionInvite>>& Entry : *UserMap)
-				{
-					if (Entry.Value->Session->SessionId == FoundSession->Get().SessionId)
-					{
-						InviteIdsToRemove.Add(Entry.Key);
-					}
-				}
-				for (const FOnlineSessionInviteIdHandle& InviteId : InviteIdsToRemove)
+				if (Entry.Value->SessionId == FoundSession->SessionId)
 				{
-					UserMap->Remove(InviteId);
+					InviteIdsToRemove.Add(Entry.Key);
 				}
 			}
+			for (const FOnlineSessionInviteIdHandle& InviteId : InviteIdsToRemove)
+			{
+				UserMap->Remove(InviteId);
+			}
+		}
 
-			Op.SetResult(FJoinSession::Result{ *FoundSession });
+		Op.SetResult(FJoinSession::Result{ });
 
-			FSessionJoined Event = { { OpParams.LocalUserId }, *FoundSession };
+		FSessionJoined Event = { { OpParams.LocalUserId }, FoundSession->SessionId };
 
-			SessionEvents.OnSessionJoined.Broadcast(Event);
-		}
+		SessionEvents.OnSessionJoined.Broadcast(Event);
 
 		// A successful join allows the client to server travel, after which RegisterPlayers will be called by the engine
-
-		// TODO: Support multiple local users joining the session
 	})
 	.Enqueue(GetSerialQueue());
 
 	return Op->GetHandle();
 }
 
-TResult<TSharedRef<const FSession>, FOnlineError> FSessionsEOSGS::BuildSessionFromInvite(const FString& InInviteId) const
+TOnlineAsyncOpHandle<FBuildSessionFromDetailsHandle> FSessionsEOSGS::BuildSessionFromInvite(const FAccountId& LocalUserId, const FString& InInviteId)
 {
 	EOS_Sessions_CopySessionHandleByInviteIdOptions CopySessionHandleByInviteIdOptions = { };
 	CopySessionHandleByInviteIdOptions.ApiVersion = EOS_SESSIONS_COPYSESSIONHANDLEBYINVITEID_API_LATEST;
@@ -1542,17 +1605,19 @@ TResult<TSharedRef<const FSession>, FOnlineError> FSessionsEOSGS::BuildSessionFr
 	EOS_EResult CopySessionHandleByInviteIdResult = EOS_Sessions_CopySessionHandleByInviteId(SessionsHandle, &CopySessionHandleByInviteIdOptions, &SessionDetailsHandle);
 	if (CopySessionHandleByInviteIdResult == EOS_EResult::EOS_Success)
 	{
-		return TResult<TSharedRef<const FSession>, FOnlineError>(MakeShared<const FSessionEOSGS>(SessionDetailsHandle));
+		return BuildSessionFromDetailsHandle({ LocalUserId, MakeShared<FSessionDetailsHandleEOSGS>(SessionDetailsHandle) });
 	}
 	else
 	{
 		UE_LOG(LogTemp, Error, TEXT("[FSessionsEOSGS::BuildSessionFromInvite] EOS_Sessions_CopySessionHandleByInviteId failed with result [%s]"), *LexToString(CopySessionHandleByInviteIdResult));
 
-		return TResult<TSharedRef<const FSession>, FOnlineError>(Errors::FromEOSResult(CopySessionHandleByInviteIdResult));
+		TOnlineAsyncOpPtr<FBuildSessionFromDetailsHandle> Operation;
+		Operation->SetError(Errors::FromEOSResult(CopySessionHandleByInviteIdResult));
+		return Operation->GetHandle();
 	}
 }
 
-TResult<TSharedRef<const FSession>, FOnlineError> FSessionsEOSGS::BuildSessionFromUIEvent(const EOS_UI_EventId& UIEventId) const
+TOnlineAsyncOpHandle<FBuildSessionFromDetailsHandle> FSessionsEOSGS::BuildSessionFromUIEvent(const FAccountId& LocalUserId, const EOS_UI_EventId& UIEventId)
 {
 	EOS_Sessions_CopySessionHandleByUiEventIdOptions CopySessionHandleByUiEventIdOptions = { };
 	CopySessionHandleByUiEventIdOptions.ApiVersion = EOS_SESSIONS_COPYSESSIONHANDLEBYUIEVENTID_API_LATEST;
@@ -1564,13 +1629,15 @@ TResult<TSharedRef<const FSession>, FOnlineError> FSessionsEOSGS::BuildSessionFr
 	EOS_EResult CopySessionHandleByUiEventIdResult = EOS_Sessions_CopySessionHandleByUiEventId(SessionsHandle, &CopySessionHandleByUiEventIdOptions, &SessionDetailsHandle);
 	if (CopySessionHandleByUiEventIdResult == EOS_EResult::EOS_Success)
 	{
-		return TResult<TSharedRef<const FSession>, FOnlineError>(MakeShared<FSessionEOSGS>(SessionDetailsHandle));
+		return BuildSessionFromDetailsHandle({ LocalUserId, MakeShared<FSessionDetailsHandleEOSGS>(SessionDetailsHandle) });
 	}
 	else
 	{
 		UE_LOG(LogTemp, Error, TEXT("[FSessionsEOSGS::BuildSessionFromUIEvent] EOS_Sessions_CopySessionHandleByUiEventId failed with result [%s]"), *LexToString(CopySessionHandleByUiEventIdResult));
 
-		return TResult<TSharedRef<const FSession>, FOnlineError>(Errors::FromEOSResult(CopySessionHandleByUiEventIdResult));
+		TOnlineAsyncOpPtr<FBuildSessionFromDetailsHandle> Operation;
+		Operation->SetError(Errors::FromEOSResult(CopySessionHandleByUiEventIdResult));
+		return Operation->GetHandle();
 	}
 }
 
diff --git a/Engine/Plugins/Online/Experimental/OnlineServicesEOSGS/Source/Public/Online/SessionsEOSGS.h b/Engine/Plugins/Online/Experimental/OnlineServicesEOSGS/Source/Public/Online/SessionsEOSGS.h
--- a/Engine/Plugins/Online/Experimental/OnlineServicesEOSGS/Source/Public/Online/SessionsEOSGS.h
+++ b/Engine/Plugins/Online/Experimental/OnlineServicesEOSGS/Source/Public/Online/SessionsEOSGS.h
@@ -129,6 +129,8 @@ struct FUpdateSessionImpl
 
 	struct Result
 	{
+		/** EOSGS Session Id for the created or modified session */
+		FString NewSessionId;
 	};
 };
 
@@ -252,19 +254,19 @@ protected:
 	FOnlineSessionInviteIdHandle CreateSessionInviteId(const FString& SessionInviteId) const;
 
 	/**
-	 * Builds a session from an invite id, using the class' Sessions Handle
+	 * Builds a session from an EOS Session Details Handle. Asynchronous due to the id resolution process
 	 */
-	TResult<TSharedRef<const FSession>, FOnlineError> BuildSessionFromInvite(const FString& InInviteId) const;
+	TOnlineAsyncOpHandle<FBuildSessionFromDetailsHandle> BuildSessionFromDetailsHandle(FBuildSessionFromDetailsHandle::Params&& Params);
 
 	/**
-	 * Builds a session from a UI event id, using the class' Sessions Handle
+	 * Builds a session from an invite id, calling BuildSessionFromDetailsHandle
 	 */
-	TResult<TSharedRef<const FSession>, FOnlineError> BuildSessionFromUIEvent(const EOS_UI_EventId& UIEventId) const;
+	TOnlineAsyncOpHandle<FBuildSessionFromDetailsHandle> BuildSessionFromInvite(const FAccountId& LocalUserId, const FString& InInviteId);
 
 	/**
-	 * Builds a session from an EOS Session Details Handle. Asynchronous due to the id resolution process
+	 * Builds a session from a UI event id, calling BuildSessionFromDetailsHandle
 	 */
-	TOnlineAsyncOpHandle<FBuildSessionFromDetailsHandle> BuildSessionFromDetailsHandle(FBuildSessionFromDetailsHandle::Params&& Params);
+	TOnlineAsyncOpHandle<FBuildSessionFromDetailsHandle> BuildSessionFromUIEvent(const FAccountId& LocalUserId, const EOS_UI_EventId& UIEventId);
 
 private:
 	// FSessionsLAN
@@ -278,7 +280,7 @@ protected:
 	EOSEventRegistrationPtr OnSessionInviteAcceptedEventRegistration;
 	EOSEventRegistrationPtr OnJoinSessionAcceptedEventRegistration;
 
-	TSharedPtr<FSessionSearchHandleEOSGS> CurrentSessionSearchHandleEOSGS;
+	TMap<FAccountId, TSharedRef<FSessionSearchHandleEOSGS>> CurrentSessionSearchHandleEOSGSUserMap;
 };
 
 namespace Meta {
@@ -296,6 +298,7 @@ BEGIN_ONLINE_STRUCT_META(FUpdateSessionImpl::Params)
 END_ONLINE_STRUCT_META()
 
 BEGIN_ONLINE_STRUCT_META(FUpdateSessionImpl::Result)
+	ONLINE_STRUCT_FIELD(FUpdateSessionImpl::Result, NewSessionId)
 END_ONLINE_STRUCT_META()
 
 BEGIN_ONLINE_STRUCT_META(FSendSingleSessionInviteImpl::Params)
diff --git a/Engine/Plugins/Online/Experimental/OnlineServicesOSSAdapter/Source/Private/Online/SessionsOSSAdapter.cpp b/Engine/Plugins/Online/Experimental/OnlineServicesOSSAdapter/Source/Private/Online/SessionsOSSAdapter.cpp
--- a/Engine/Plugins/Online/Experimental/OnlineServicesOSSAdapter/Source/Private/Online/SessionsOSSAdapter.cpp
+++ b/Engine/Plugins/Online/Experimental/OnlineServicesOSSAdapter/Source/Private/Online/SessionsOSSAdapter.cpp
@@ -35,24 +35,6 @@ EOnlineComparisonOp::Type GetV1SessionSearchComparisonOp(const ESchemaAttributeC
 	return EOnlineComparisonOp::Equals;
 }
 
-ESessionState GetV2SessionState(const EOnlineSessionState::Type InValue)
-{
-	switch (InValue)
-	{
-	case EOnlineSessionState::NoSession: return ESessionState::Invalid;
-	case EOnlineSessionState::Creating: return ESessionState::Creating;
-	case EOnlineSessionState::Pending: return ESessionState::Joining;
-	case EOnlineSessionState::Starting: return ESessionState::Valid;
-	case EOnlineSessionState::InProgress: return ESessionState::Valid;
-	case EOnlineSessionState::Ending: return ESessionState::Valid;
-	case EOnlineSessionState::Ended: return ESessionState::Valid;
-	case EOnlineSessionState::Destroying: return ESessionState::Destroying;
-	}
-
-	checkNoEntry();
-	return ESessionState::Invalid;
-}
-
 FSchemaVariant GetV2SessionVariant(const FVariantData& InValue)
 {
 	FSchemaVariant Result;
@@ -296,13 +278,19 @@ void FSessionsOSSAdapter::Initialize()
 	{
 		FOnlineServicesOSSAdapter& ServicesOSSAdapter = static_cast<FOnlineServicesOSSAdapter&>(Services);
 
-		FSessionInvite SessionInvite;
-		SessionInvite.SenderId = ServicesOSSAdapter.GetAccountIdRegistry().FindOrAddHandle(FromId.AsShared());
-		SessionInvite.RecipientId = ServicesOSSAdapter.GetAccountIdRegistry().FindOrAddHandle(UserId.AsShared());
-		SessionInvite.InviteId = GetSessionInviteIdRegistry().BasicRegistry.FindOrAddHandle(AppId);
-		SessionInvite.Session = BuildV2Session(&InviteResult.Session);
+		TSharedRef<FSession> Session = BuildV2Session(&InviteResult.Session);
+
+		TSharedRef<FSessionInvite> SessionInvite = MakeShared<FSessionInvite>();
+		SessionInvite->SenderId = ServicesOSSAdapter.GetAccountIdRegistry().FindOrAddHandle(FromId.AsShared());
+		SessionInvite->RecipientId = ServicesOSSAdapter.GetAccountIdRegistry().FindOrAddHandle(UserId.AsShared());
+		SessionInvite->InviteId = GetSessionInviteIdRegistry().BasicRegistry.FindOrAddHandle(AppId);
+		SessionInvite->SessionId = Session->SessionId;
 
-		FSessionInviteReceived Event{ ServicesOSSAdapter.GetAccountIdRegistry().FindOrAddHandle(UserId.AsShared()), MakeShared<FSessionInvite>(MoveTemp(SessionInvite)) };
+		SessionInvitesUserMap.FindOrAdd(SessionInvite->RecipientId).Emplace(SessionInvite->InviteId, SessionInvite);
+
+		AllSessionsById.Emplace(Session->SessionId, Session);
+
+		FSessionInviteReceived Event{ ServicesOSSAdapter.GetAccountIdRegistry().FindOrAddHandle(UserId.AsShared()), SessionInvite };
 
 		SessionEvents.OnSessionInviteReceived.Broadcast(Event);
 	}));
@@ -320,7 +308,7 @@ void FSessionsOSSAdapter::Initialize()
 
 		FUISessionJoinRequested Event {
 			ServicesOSSAdapter.GetAccountIdRegistry().FindOrAddHandle(LocalUserId),
-			TResult<TSharedRef<const FSession>, FOnlineError>(BuildV2Session(&InviteResult.Session)),
+			TResult<FOnlineSessionIdHandle, FOnlineError>(GetSessionIdRegistry().FindOrAddHandle(InviteResult.Session.SessionInfo->GetSessionId().AsShared())),
 			EUISessionJoinRequestedSource::FromInvitation
 		};
 
@@ -329,23 +317,16 @@ void FSessionsOSSAdapter::Initialize()
 
 	SessionsInterface->AddOnSessionParticipantRemovedDelegate_Handle(FOnSessionParticipantRemovedDelegate::CreateLambda([this](FName SessionName, const FUniqueNetId& TargetUniqueNetId)
 	{
-		// We won't transmit events for a session or member that doesn't exist
-		const TOnlineResult<FGetSessionByName> GetSessionByNameResult = GetSessionByName({ SessionName });
-		if (GetSessionByNameResult.IsOk())
-		{
-			const TSharedRef<const FSession>& FoundSession = GetSessionByNameResult.GetOkValue().Session;
-
-			const FOnlineServicesOSSAdapter& ServicesOSSAdapter = static_cast<FOnlineServicesOSSAdapter&>(Services);
-			const FUniqueNetIdRef TargetUniqueNetIdRef = TargetUniqueNetId.AsShared();
-			const FAccountId TargetAccountIdHandle = ServicesOSSAdapter.GetAccountIdRegistry().FindOrAddHandle(TargetUniqueNetIdRef);
+		const FOnlineServicesOSSAdapter& ServicesOSSAdapter = static_cast<FOnlineServicesOSSAdapter&>(Services);
+		const FUniqueNetIdRef TargetUniqueNetIdRef = TargetUniqueNetId.AsShared();
+		const FAccountId TargetAccountIdHandle = ServicesOSSAdapter.GetAccountIdRegistry().FindOrAddHandle(TargetUniqueNetIdRef);
 
-			FSessionSettingsUpdate SettingsUpdate;
-			SettingsUpdate.RemovedSessionMembers.Add(TargetAccountIdHandle);
+		FSessionSettingsUpdate SettingsUpdate;
+		SettingsUpdate.RemovedSessionMembers.Add(TargetAccountIdHandle);
 
-			const FSessionUpdated Event { FoundSession , SettingsUpdate };
+		const FSessionUpdated Event { SessionName , SettingsUpdate };
 
-			SessionEvents.OnSessionUpdated.Broadcast(Event);
-		}
+		SessionEvents.OnSessionUpdated.Broadcast(Event);
 	}));
 
 	SessionsInterface->AddOnSessionParticipantsChangeDelegate_Handle(FOnSessionParticipantsChangeDelegate::CreateLambda([this](FName SessionName, const FUniqueNetId& TargetUniqueNetId, bool bJoined)
@@ -376,7 +357,7 @@ void FSessionsOSSAdapter::Initialize()
 				SettingsUpdate.RemovedSessionMembers.Add(TargetAccountIdHandle);
 			}
 
-			const FSessionUpdated Event { FoundSession , SettingsUpdate };
+			const FSessionUpdated Event { SessionName , SettingsUpdate };
 
 			SessionEvents.OnSessionUpdated.Broadcast(Event);
 		}
@@ -415,7 +396,7 @@ void FSessionsOSSAdapter::Initialize()
 				SettingsUpdate.UpdatedSessionMembers.Emplace(TargetAccountIdHandle, SessionMemberUpdate);
 			}
 
-			const FSessionUpdated Event{ FoundSession , SettingsUpdate };
+			const FSessionUpdated Event{ SessionName , SettingsUpdate };
 
 			SessionEvents.OnSessionUpdated.Broadcast(Event);
 		}
@@ -433,53 +414,6 @@ void FSessionsOSSAdapter::Shutdown()
 	Super::Shutdown();
 }
 
-TOnlineResult<FGetAllSessions> FSessionsOSSAdapter::GetAllSessions(FGetAllSessions::Params&& Params) const
-{
-	TArray<TSharedRef<const FSession>> Sessions;
-
-	for (const FName SessionName : RegisteredSessionNames)
-	{
-		if (FNamedOnlineSession* NamedSession = SessionsInterface->GetNamedSession(SessionName))
-		{
-			Sessions.Add(BuildV2Session(NamedSession));
-		}
-	}
-
-	return TOnlineResult<FGetAllSessions>({ Sessions });
-}
-
-TOnlineResult<FGetSessionByName> FSessionsOSSAdapter::GetSessionByName(FGetSessionByName::Params&& Params) const
-{
-	if (FNamedOnlineSession* NamedSession = SessionsInterface->GetNamedSession(Params.LocalName))
-	{
-		return TOnlineResult<FGetSessionByName>({ BuildV2Session(NamedSession) });
-	}
-	else
-	{
-		return TOnlineResult<FGetSessionByName>(Errors::NotFound());
-	}
-}
-
-TOnlineResult<FGetSessionById> FSessionsOSSAdapter::GetSessionById(FGetSessionById::Params&& Params) const
-{
-	FOnlineServicesOSSAdapter& ServicesOSSAdapter = static_cast<FOnlineServicesOSSAdapter&>(Services);
-
-	for (const FName SessionName : RegisteredSessionNames)
-	{
-		if (FNamedOnlineSession* NamedSession = SessionsInterface->GetNamedSession(SessionName))
-		{
-			FOnlineSessionIdHandle SessionIdHandle = GetSessionIdRegistry().FindOrAddHandle(NamedSession->SessionInfo->GetSessionId().AsShared());
-
-			if (SessionIdHandle == Params.IdHandle)
-			{
-				return TOnlineResult<FGetSessionById>({ BuildV2Session(NamedSession) });
-			}
-		}
-	}
-
-	return TOnlineResult<FGetSessionById>(Errors::NotFound());
-}
-
 TOnlineAsyncOpHandle<FCreateSession> FSessionsOSSAdapter::CreateSession(FCreateSession::Params&& Params)
 {
 	TOnlineAsyncOpRef<FCreateSession> Op = GetOp<FCreateSession>(MoveTemp(Params));
@@ -517,13 +451,19 @@ TOnlineAsyncOpHandle<FCreateSession> FSessionsOSSAdapter::CreateSession(FCreateS
 					return;
 				}
 
-				RegisteredSessionNames.Add(SessionName);
+				const FCreateSession::Params& OpParams = Op->GetParams();
 
 				const FNamedOnlineSession* V1Session = SessionsInterface->GetNamedSession(SessionName);
 
 				TSharedRef<FSession> V2Session = BuildV2Session(V1Session);
 
-				Op->SetResult({ V2Session });
+				AllSessionsById.Emplace(V2Session->SessionId, V2Session);
+
+				LocalSessionsByName.Emplace(OpParams.SessionName, V2Session->SessionId);
+
+				NamedSessionUserMap.FindOrAdd(OpParams.LocalUserId).AddUnique(OpParams.SessionName);
+
+				Op->SetResult({ });
 			}
 		});
 
@@ -565,11 +505,17 @@ TOnlineAsyncOpHandle<FUpdateSession> FSessionsOSSAdapter::UpdateSession(FUpdateS
 
 				const FUpdateSession::Params& OpParams = Op->GetParams();
 
-				if (TSharedRef<FSession>* FoundSession = LocalSessionsByName.Find(OpParams.SessionName))
+				const TOnlineResult<FGetMutableSessionByName> GetMutableSessionByNameResult = GetMutableSessionByName({ SessionName });
+				if (GetMutableSessionByNameResult.IsOk())
 				{
-					Op->SetResult({ *FoundSession });
+					const TSharedRef<FSession>& FoundSession = GetMutableSessionByNameResult.GetOkValue().Session;
+
+					FSessionSettings& UpdatedV2Settings = FoundSession->SessionSettings;
+					UpdatedV2Settings += OpParams.Mutations;
 
-					FSessionUpdated Event { *FoundSession, OpParams.Mutations };
+					Op->SetResult({ });
+
+					FSessionUpdated Event { SessionName, OpParams.Mutations };
 
 					SessionEvents.OnSessionUpdated.Broadcast(Event);
 				}
@@ -580,7 +526,10 @@ TOnlineAsyncOpHandle<FUpdateSession> FSessionsOSSAdapter::UpdateSession(FUpdateS
 			}
 		});
 
-		const TSharedRef<FSession>& FoundSession = LocalSessionsByName.FindChecked(OpParams.SessionName);
+		const TOnlineResult<FGetSessionByName> GetSessionByNameResult = GetSessionByName({ OpParams.SessionName });
+		check(GetSessionByNameResult.IsOk());
+		
+		const TSharedRef<const FSession>& FoundSession = GetSessionByNameResult.GetOkValue().Session;
 
 		// We will update a copy here, and wait until the operation has completed successfully to update our local data
 		FSessionSettings UpdatedV2Settings = FoundSession->SessionSettings;
@@ -623,7 +572,23 @@ TOnlineAsyncOpHandle<FLeaveSession> FSessionsOSSAdapter::LeaveSession(FLeaveSess
 					return;
 				}
 
-				RegisteredSessionNames.Remove(SessionName);
+				const FLeaveSession::Params& OpParams = Op->GetParams();
+
+				TOnlineResult<FGetSessionByName> GetSessionByNameResult = GetSessionByName({ OpParams.SessionName });
+				if (GetSessionByNameResult.IsOk())
+				{
+					TSharedRef<const FSession> FoundSession = GetSessionByNameResult.GetOkValue().Session;
+
+					NamedSessionUserMap.FindChecked(OpParams.LocalUserId).Remove(OpParams.SessionName);
+
+					if (FoundSession->SessionSettings.bPresenceEnabled)
+					{
+						PresenceSessionsUserMap.Remove(OpParams.LocalUserId);
+					}
+
+					ClearSessionByName(OpParams.SessionName);
+					ClearSessionById(FoundSession->SessionId);
+				}
 
 				Op->SetResult({ });
 
@@ -697,9 +662,16 @@ TOnlineAsyncOpHandle<FFindSessions> FSessionsOSSAdapter::FindSessions(FFindSessi
 						{
 							TSharedRef<FOnlineSessionSearch>& PendingV1SessionSearch = PendingV1SessionSearchesPerUser.FindChecked(Op->GetParams().LocalUserId);
 
-							TArray<TSharedRef<const FSession>> FoundSessions = BuildV2SessionSearchResults(PendingV1SessionSearch->SearchResults);
+							TArray<TSharedRef<FSession>> FoundSessions = BuildV2SessionSearchResults(PendingV1SessionSearch->SearchResults);
+
+							TArray<FOnlineSessionIdHandle> SearchResults = SearchResultsUserMap.FindOrAdd(Op->GetParams().LocalUserId);
+							for (const TSharedRef<FSession>& FoundSession : FoundSessions)
+							{
+								AllSessionsById.Emplace(FoundSession->SessionId, FoundSession);
+								SearchResults.Add(FoundSession->SessionId);
+							}
 
-							Op->SetResult({ MoveTemp(FoundSessions) });
+							Op->SetResult({ SearchResults });
 						}
 						else
 						{
@@ -753,9 +725,15 @@ TOnlineAsyncOpHandle<FFindSessions> FSessionsOSSAdapter::FindSessions(FFindSessi
 				{
 					if (bWasSuccessful)
 					{
-						TArray<TSharedRef<const FSession>> FoundSessions = BuildV2SessionSearchResults(FriendSearchResults);
+						TArray<TSharedRef<FSession>> FoundSessions = BuildV2SessionSearchResults(FriendSearchResults);
+						TArray<FOnlineSessionIdHandle> SearchResults = SearchResultsUserMap.FindOrAdd(Op->GetParams().LocalUserId);
+						for (const TSharedRef<FSession>& FoundSession : FoundSessions)
+						{
+							AllSessionsById.Emplace(FoundSession->SessionId, FoundSession);
+							SearchResults.Add(FoundSession->SessionId);
+						}
 
-						Op->SetResult({ MoveTemp(FoundSessions) });
+						Op->SetResult({ SearchResults });
 					}
 					else
 					{
@@ -786,9 +764,15 @@ TOnlineAsyncOpHandle<FFindSessions> FSessionsOSSAdapter::FindSessions(FFindSessi
 					{
 						TSharedRef<FOnlineSessionSearch>& PendingV1SessionSearch = PendingV1SessionSearchesPerUser.FindChecked(Op->GetParams().LocalUserId);
 
-						TArray<TSharedRef<const FSession>> FoundSessions = BuildV2SessionSearchResults(PendingV1SessionSearch->SearchResults);
+						TArray<TSharedRef<FSession>> FoundSessions = BuildV2SessionSearchResults(PendingV1SessionSearch->SearchResults);
+						TArray<FOnlineSessionIdHandle> SearchResults = SearchResultsUserMap.FindOrAdd(Op->GetParams().LocalUserId);
+						for (const TSharedRef<FSession>& FoundSession : FoundSessions)
+						{
+							AllSessionsById.Emplace(FoundSession->SessionId, FoundSession);
+							SearchResults.Add(FoundSession->SessionId);
+						}
 
-						Op->SetResult({ MoveTemp(FoundSessions) });
+						Op->SetResult({ SearchResults });
 					}
 					else
 					{
@@ -890,11 +874,19 @@ TOnlineAsyncOpHandle<FStartMatchmaking> FSessionsOSSAdapter::StartMatchmaking(FS
 					return;
 				}
 
-				TSharedRef<FSession> NewSession = BuildV2Session(SessionsInterface->GetNamedSession(SessionName));
+				const FStartMatchmaking::Params& OpParams = Op->GetParams();
+
+				TSharedRef<FSession> V2Session = BuildV2Session(SessionsInterface->GetNamedSession(SessionName));
+
+				AllSessionsById.Emplace(V2Session->SessionId, V2Session);
 
-				Op->SetResult({ NewSession });
+				LocalSessionsByName.Emplace(OpParams.SessionName, V2Session->SessionId);
 
-				FSessionJoined Event{ { Op->GetParams().LocalUserId }, NewSession };
+				NamedSessionUserMap.FindOrAdd(OpParams.LocalUserId).AddUnique(OpParams.SessionName);
+
+				Op->SetResult({ });
+
+				FSessionJoined Event{ { Op->GetParams().LocalUserId }, V2Session->SessionId };
 
 				SessionEvents.OnSessionJoined.Broadcast(Event);
 
@@ -963,13 +955,38 @@ TOnlineAsyncOpHandle<FJoinSession> FSessionsOSSAdapter::JoinSession(FJoinSession
 					return;
 				}
 
-				RegisteredSessionNames.Add(SessionName);
+				TSharedRef<const FSession> FoundSession = GetSessionByIdResult.GetOkValue().Session;
 
-				const TSharedRef<const FSession>& FoundSession = GetSessionByIdResult.GetOkValue().Session;
+				LocalSessionsByName.Emplace(OpParams.SessionName, OpParams.SessionId);
 
-				Op->SetResult({ FoundSession });
+				NamedSessionUserMap.FindOrAdd(OpParams.LocalUserId).AddUnique(OpParams.SessionName);
+
+				if (FoundSession->SessionSettings.bPresenceEnabled)
+				{
+					FOnlineSessionIdHandle& PresenceSessionId = PresenceSessionsUserMap.FindOrAdd(OpParams.LocalUserId);
+					PresenceSessionId = FoundSession->SessionId;
+				}
+
+				// After successfully joining a session, we'll remove all related invites if any are found
+				if (TMap<FOnlineSessionInviteIdHandle, TSharedRef<FSessionInvite>>* UserMap = SessionInvitesUserMap.Find(OpParams.LocalUserId))
+				{
+					TArray<FOnlineSessionInviteIdHandle> InviteIdsToRemove;
+					for (const TPair<FOnlineSessionInviteIdHandle, TSharedRef<FSessionInvite>>& Entry : *UserMap)
+					{
+						if (Entry.Value->SessionId == FoundSession->SessionId)
+						{
+							InviteIdsToRemove.Add(Entry.Key);
+						}
+					}
+					for (const FOnlineSessionInviteIdHandle& InviteId : InviteIdsToRemove)
+					{
+						UserMap->Remove(InviteId);
+					}
+				}
+
+				Op->SetResult({ });
 
-				FSessionJoined Event{ { OpParams.LocalUserId }, FoundSession };
+				FSessionJoined Event{ { OpParams.LocalUserId }, FoundSession->SessionId };
 
 				SessionEvents.OnSessionJoined.Broadcast(Event);
 			}
@@ -1318,16 +1335,15 @@ TSharedRef<FSession> FSessionsOSSAdapter::BuildV2Session(const FOnlineSession* I
 	if (const FNamedOnlineSession* NamedInSession = static_cast<const FNamedOnlineSession*>(InSession))
 	{
 		Session->OwnerUserId = ServicesOSSAdapter.GetAccountIdRegistry().FindOrAddHandle(NamedInSession->LocalOwnerId->AsShared());
-		Session->CurrentState = GetV2SessionState(NamedInSession->SessionState);
 		WriteV2SessionSettingsFromV1NamedSession(NamedInSession, Session->SessionSettings);
 	}
 
 	return Session;
 }
 
-TArray<TSharedRef<const FSession>> FSessionsOSSAdapter::BuildV2SessionSearchResults(const TArray<FOnlineSessionSearchResult>& SessionSearchResults) const
+TArray<TSharedRef<FSession>> FSessionsOSSAdapter::BuildV2SessionSearchResults(const TArray<FOnlineSessionSearchResult>& SessionSearchResults) const
 {
-	TArray<TSharedRef<const FSession>> FoundSessions;
+	TArray<TSharedRef<FSession>> FoundSessions;
 
 	for (const FOnlineSessionSearchResult& SearchResult : SessionSearchResults)
 	{
diff --git a/Engine/Plugins/Online/Experimental/OnlineServicesOSSAdapter/Source/Public/Online/SessionsOSSAdapter.h b/Engine/Plugins/Online/Experimental/OnlineServicesOSSAdapter/Source/Public/Online/SessionsOSSAdapter.h
--- a/Engine/Plugins/Online/Experimental/OnlineServicesOSSAdapter/Source/Public/Online/SessionsOSSAdapter.h
+++ b/Engine/Plugins/Online/Experimental/OnlineServicesOSSAdapter/Source/Public/Online/SessionsOSSAdapter.h
@@ -63,9 +63,6 @@ public:
 	virtual void Shutdown() override;
 
 	// ISessions
-	virtual TOnlineResult<FGetAllSessions> GetAllSessions(FGetAllSessions::Params&& Params) const override;
-	virtual TOnlineResult<FGetSessionByName> GetSessionByName(FGetSessionByName::Params&& Params) const override;
-	virtual TOnlineResult<FGetSessionById> GetSessionById(FGetSessionById::Params&& Params) const override;
 	virtual TOnlineAsyncOpHandle<FCreateSession> CreateSession(FCreateSession::Params&& Params) override;
 	virtual TOnlineAsyncOpHandle<FUpdateSession> UpdateSession(FUpdateSession::Params&& Params) override;
 	virtual TOnlineAsyncOpHandle<FLeaveSession> LeaveSession(FLeaveSession::Params&& Params) override;
@@ -92,7 +89,7 @@ private:
 	/** Builds a V2 Session data type from a V1 equivalent */
 	TSharedRef<FSession> BuildV2Session(const FOnlineSession* InSession) const;
 	/** Builds a V2 Session Search Results array from the passed V1 equivalent types */
-	TArray<TSharedRef<const FSession>> BuildV2SessionSearchResults(const TArray<FOnlineSessionSearchResult>& SessionSearchResults) const;
+	TArray<TSharedRef<FSession>> BuildV2SessionSearchResults(const TArray<FOnlineSessionSearchResult>& SessionSearchResults) const;
 
 	FOnlineSessionIdRegistryOSSAdapter& GetSessionIdRegistry() const;
 	FOnlineSessionInviteIdRegistryOSSAdapter& GetSessionInviteIdRegistry() const;
@@ -101,9 +98,6 @@ protected:
 	IOnlineSessionPtr SessionsInterface;
 	IOnlineIdentityPtr IdentityInterface;
 
-	/** Record of all names for registered sessions. Needed for session retrieval methods */
-	TSet<FName> RegisteredSessionNames;
-
 	/** Cache of session search information to be used by ongoing session searches, indexed per user. Entries are removed upon search completion. */
 	TMap<FAccountId, TSharedRef<FOnlineSessionSearch>> PendingV1SessionSearchesPerUser;
 };
