OSSv2: Rename Handle, UserId to AccountId when use FAccountId.

diff --git a/Engine/Plugins/Online/Experimental/OnlineServicesEOSGS/Source/Private/Online/AchievementsEOSGS.cpp b/Engine/Plugins/Online/Experimental/OnlineServicesEOSGS/Source/Private/Online/AchievementsEOSGS.cpp
--- a/Engine/Plugins/Online/Experimental/OnlineServicesEOSGS/Source/Private/Online/AchievementsEOSGS.cpp
+++ b/Engine/Plugins/Online/Experimental/OnlineServicesEOSGS/Source/Private/Online/AchievementsEOSGS.cpp
@@ -33,15 +33,15 @@ void FAchievementsEOSGS::Initialize()
 	{
 		FAchievementsEOSGS* This = reinterpret_cast<FAchievementsEOSGS*>(Data->ClientData);
 
-		const FAccountId LocalUserId = FindAccountId(Data->UserId);
-		if (LocalUserId.IsValid())
+		const FAccountId LocalAccountId = FindAccountId(Data->UserId);
+		if (LocalAccountId.IsValid())
 		{
 			const FString AchievementId = UTF8_TO_TCHAR(Data->AchievementId);
 			const FDateTime UnlockTime = FDateTime::FromUnixTimestamp(Data->UnlockTime);
 
-			UE_LOG(LogTemp, Verbose, TEXT("EOS_Achievements_OnAchievementsUnlocked User=%s AchievId=%s"), *ToLogString(LocalUserId), *AchievementId);
+			UE_LOG(LogTemp, Verbose, TEXT("EOS_Achievements_OnAchievementsUnlocked User=%s AchievId=%s"), *ToLogString(LocalAccountId), *AchievementId);
 
-			if (FAchievementStateMap* LocalUserAchievementStates = This->AchievementStates.Find(LocalUserId))
+			if (FAchievementStateMap* LocalUserAchievementStates = This->AchievementStates.Find(LocalAccountId))
 			{
 				if (FAchievementState* AchievementState = LocalUserAchievementStates->Find(AchievementId))
 				{
@@ -51,7 +51,7 @@ void FAchievementsEOSGS::Initialize()
 			}
 			
 			FAchievementStateUpdated EventParams;
-			EventParams.LocalUserId = LocalUserId;
+			EventParams.LocalAccountId = LocalAccountId;
 			EventParams.AchievementIds.Emplace(AchievementId);
 
 			This->OnAchievementStateUpdatedEvent.Broadcast(EventParams);
@@ -74,7 +74,7 @@ TOnlineAsyncOpHandle<FQueryAchievementDefinitions> FAchievementsEOSGS::QueryAchi
 {
 	TOnlineAsyncOpRef<FQueryAchievementDefinitions> Op = GetOp<FQueryAchievementDefinitions>(MoveTemp(Params));
 
-	if (!Services.Get<FAuthEOSGS>()->IsLoggedIn(Op->GetParams().LocalUserId))
+	if (!Services.Get<FAuthEOSGS>()->IsLoggedIn(Op->GetParams().LocalAccountId))
 	{
 		Op->SetError(Errors::InvalidUser());
 		return Op->GetHandle();
@@ -87,7 +87,7 @@ TOnlineAsyncOpHandle<FQueryAchievementDefinitions> FAchievementsEOSGS::QueryAchi
 			EOS_Achievements_QueryDefinitionsOptions Options = {};
 			Options.ApiVersion = EOS_ACHIEVEMENTS_QUERYDEFINITIONS_API_LATEST;
 			static_assert(EOS_ACHIEVEMENTS_QUERYDEFINITIONS_API_LATEST == 3, "EOS_Achievements_QueryDefinitionsOptions updated, check new fields");
-			Options.LocalUserId = GetProductUserIdChecked(InAsyncOp.GetParams().LocalUserId);
+			Options.LocalUserId = GetProductUserIdChecked(InAsyncOp.GetParams().LocalAccountId);
 
 			EOS_Async(EOS_Achievements_QueryDefinitions, AchievementsHandle, Options, MoveTemp(Promise));
 		})
@@ -169,7 +169,7 @@ TOnlineAsyncOpHandle<FQueryAchievementDefinitions> FAchievementsEOSGS::QueryAchi
 
 TOnlineResult<FGetAchievementIds> FAchievementsEOSGS::GetAchievementIds(FGetAchievementIds::Params&& Params)
 {
-	if (!Services.Get<FAuthEOSGS>()->IsLoggedIn(Params.LocalUserId))
+	if (!Services.Get<FAuthEOSGS>()->IsLoggedIn(Params.LocalAccountId))
 	{
 		return TOnlineResult<FGetAchievementIds>(Errors::InvalidUser());
 	}
@@ -187,7 +187,7 @@ TOnlineResult<FGetAchievementIds> FAchievementsEOSGS::GetAchievementIds(FGetAchi
 
 TOnlineResult<FGetAchievementDefinition> FAchievementsEOSGS::GetAchievementDefinition(FGetAchievementDefinition::Params&& Params)
 {
-	if (!Services.Get<FAuthEOSGS>()->IsLoggedIn(Params.LocalUserId))
+	if (!Services.Get<FAuthEOSGS>()->IsLoggedIn(Params.LocalAccountId))
 	{
 		return TOnlineResult<FGetAchievementDefinition>(Errors::InvalidUser());
 	}
@@ -211,7 +211,7 @@ TOnlineAsyncOpHandle<FQueryAchievementStates> FAchievementsEOSGS::QueryAchieveme
 {
 	TOnlineAsyncOpRef<FQueryAchievementStates> Op = GetOp<FQueryAchievementStates>(MoveTemp(Params));
 
-	if (!Services.Get<FAuthEOSGS>()->IsLoggedIn(Op->GetParams().LocalUserId))
+	if (!Services.Get<FAuthEOSGS>()->IsLoggedIn(Op->GetParams().LocalAccountId))
 	{
 		Op->SetError(Errors::InvalidUser());
 		return Op->GetHandle();
@@ -224,14 +224,14 @@ TOnlineAsyncOpHandle<FQueryAchievementStates> FAchievementsEOSGS::QueryAchieveme
 		return Op->GetHandle();
 	}
 
-	if (!AchievementStates.Contains(Op->GetParams().LocalUserId))
+	if (!AchievementStates.Contains(Op->GetParams().LocalAccountId))
 	{
 		Op->Then([this](TOnlineAsyncOp<FQueryAchievementStates>& InAsyncOp, TPromise<const EOS_Achievements_OnQueryPlayerAchievementsCompleteCallbackInfo*>&& Promise)
 		{
 			EOS_Achievements_QueryPlayerAchievementsOptions Options = {};
 			Options.ApiVersion = EOS_ACHIEVEMENTS_QUERYPLAYERACHIEVEMENTS_API_LATEST;
 			static_assert(EOS_ACHIEVEMENTS_QUERYPLAYERACHIEVEMENTS_API_LATEST == 2, "EOS_Achievements_QueryPlayerAchievementsOptions updated, check new fields");
-			Options.LocalUserId = GetProductUserIdChecked(InAsyncOp.GetParams().LocalUserId);
+			Options.LocalUserId = GetProductUserIdChecked(InAsyncOp.GetParams().LocalAccountId);
 			Options.TargetUserId = Options.LocalUserId;
 
 			EOS_Async(EOS_Achievements_QueryPlayerAchievements, AchievementsHandle, Options, MoveTemp(Promise));
@@ -249,7 +249,7 @@ TOnlineAsyncOpHandle<FQueryAchievementStates> FAchievementsEOSGS::QueryAchieveme
 			EOS_Achievements_GetPlayerAchievementCountOptions GetCountOptions = {};
 			GetCountOptions.ApiVersion = EOS_ACHIEVEMENTS_GETPLAYERACHIEVEMENTCOUNT_API_LATEST;
 			static_assert(EOS_ACHIEVEMENTS_GETPLAYERACHIEVEMENTCOUNT_API_LATEST == 1, "EOS_Achievements_GetPlayerAchievementCountOptions updated, check new fields");
-			GetCountOptions.UserId = GetProductUserIdChecked(InAsyncOp.GetParams().LocalUserId);
+			GetCountOptions.UserId = GetProductUserIdChecked(InAsyncOp.GetParams().LocalAccountId);
 
 			const uint32_t AchievementCount = EOS_Achievements_GetPlayerAchievementCount(AchievementsHandle, &GetCountOptions);
 
@@ -259,7 +259,7 @@ TOnlineAsyncOpHandle<FQueryAchievementStates> FAchievementsEOSGS::QueryAchieveme
 				EOS_Achievements_CopyPlayerAchievementByIndexOptions CopyOptions = {};
 				CopyOptions.ApiVersion = EOS_ACHIEVEMENTS_COPYPLAYERACHIEVEMENTBYINDEX_API_LATEST;
 				static_assert(EOS_ACHIEVEMENTS_COPYPLAYERACHIEVEMENTBYINDEX_API_LATEST == 2, "EOS_Achievements_CopyPlayerAchievementByIndexOptions updated, check new fields");
-				CopyOptions.LocalUserId = GetProductUserIdChecked(InAsyncOp.GetParams().LocalUserId);
+				CopyOptions.LocalUserId = GetProductUserIdChecked(InAsyncOp.GetParams().LocalAccountId);
 				CopyOptions.TargetUserId = CopyOptions.LocalUserId;
 				CopyOptions.AchievementIndex = AchievementIdx;
 
@@ -287,7 +287,7 @@ TOnlineAsyncOpHandle<FQueryAchievementStates> FAchievementsEOSGS::QueryAchieveme
 
 				EOS_Achievements_PlayerAchievement_Release(EOSPlayerAchievement);
 			}
-			AchievementStates.Emplace(InAsyncOp.GetParams().LocalUserId, MoveTemp(NewAchievementStates));
+			AchievementStates.Emplace(InAsyncOp.GetParams().LocalAccountId, MoveTemp(NewAchievementStates));
 			InAsyncOp.SetResult({});
 		})
 		.Enqueue(GetSerialQueue());
@@ -302,12 +302,12 @@ TOnlineAsyncOpHandle<FQueryAchievementStates> FAchievementsEOSGS::QueryAchieveme
 
 TOnlineResult<FGetAchievementState> FAchievementsEOSGS::GetAchievementState(FGetAchievementState::Params&& Params) const
 {
-	if (!Services.Get<FAuthEOSGS>()->IsLoggedIn(Params.LocalUserId))
+	if (!Services.Get<FAuthEOSGS>()->IsLoggedIn(Params.LocalAccountId))
 	{
 		return TOnlineResult<FGetAchievementState>(Errors::InvalidUser());
 	}
 
-	const FAchievementStateMap* LocalUserAchievementStates = AchievementStates.Find(Params.LocalUserId);
+	const FAchievementStateMap* LocalUserAchievementStates = AchievementStates.Find(Params.LocalAccountId);
 	if (!LocalUserAchievementStates)
 	{
 		// Call QueryAchievementStates first
@@ -327,7 +327,7 @@ TOnlineAsyncOpHandle<FUnlockAchievements> FAchievementsEOSGS::UnlockAchievements
 {
 	TOnlineAsyncOpRef<FUnlockAchievements> Op = GetOp<FUnlockAchievements>(MoveTemp(Params));
 
-	if (!Services.Get<FAuthEOSGS>()->IsLoggedIn(Op->GetParams().LocalUserId))
+	if (!Services.Get<FAuthEOSGS>()->IsLoggedIn(Op->GetParams().LocalAccountId))
 	{
 		Op->SetError(Errors::InvalidUser());
 		return Op->GetHandle();
@@ -339,7 +339,7 @@ TOnlineAsyncOpHandle<FUnlockAchievements> FAchievementsEOSGS::UnlockAchievements
 		return Op->GetHandle();
 	}
 
-	FAchievementStateMap* LocalUserAchievementStates = AchievementStates.Find(Op->GetParams().LocalUserId);
+	FAchievementStateMap* LocalUserAchievementStates = AchievementStates.Find(Op->GetParams().LocalAccountId);
 	if (!LocalUserAchievementStates)
 	{
 		// Call QueryAchievementStates first
@@ -379,7 +379,7 @@ TOnlineAsyncOpHandle<FUnlockAchievements> FAchievementsEOSGS::UnlockAchievements
 		EOS_Achievements_UnlockAchievementsOptions Options = {};
 		Options.ApiVersion = EOS_ACHIEVEMENTS_UNLOCKACHIEVEMENTS_API_LATEST;
 		static_assert(EOS_ACHIEVEMENTS_UNLOCKACHIEVEMENTS_API_LATEST == 1, "EOS_Achievements_UnlockAchievementsOptions updated, check new fields");
-		Options.UserId = GetProductUserIdChecked(InAsyncOp.GetParams().LocalUserId);
+		Options.UserId = GetProductUserIdChecked(InAsyncOp.GetParams().LocalAccountId);
 		Options.AchievementIds = AchievementIdPtrs.GetData();
 		Options.AchievementsCount = AchievementIdPtrs.Num();
 
@@ -394,7 +394,7 @@ TOnlineAsyncOpHandle<FUnlockAchievements> FAchievementsEOSGS::UnlockAchievements
 			const FUnlockAchievements::Params& Params = InAsyncOp.GetParams();
 			const FDateTime TimeNow = FDateTime::Now();
 
-			FAchievementStateMap* LocalUserAchievementStates = AchievementStates.Find(Params.LocalUserId);
+			FAchievementStateMap* LocalUserAchievementStates = AchievementStates.Find(Params.LocalAccountId);
 			if (ensure(LocalUserAchievementStates))
 			{
 				for(const FString& AchievementId : Params.AchievementIds)
diff --git a/Engine/Plugins/Online/Experimental/OnlineServicesEOSGS/Source/Private/Online/AuthEOSGS.cpp b/Engine/Plugins/Online/Experimental/OnlineServicesEOSGS/Source/Private/Online/AuthEOSGS.cpp
--- a/Engine/Plugins/Online/Experimental/OnlineServicesEOSGS/Source/Private/Online/AuthEOSGS.cpp
+++ b/Engine/Plugins/Online/Experimental/OnlineServicesEOSGS/Source/Private/Online/AuthEOSGS.cpp
@@ -473,9 +473,9 @@ TSharedPtr<FAccountInfoEOS> FAccountInfoRegistryEOS::Find(FPlatformUserId Platfo
 	return StaticCastSharedPtr<FAccountInfoEOS>(Super::Find(PlatformUserId));
 }
 
-TSharedPtr<FAccountInfoEOS> FAccountInfoRegistryEOS::Find(FAccountId AccountIdHandle) const
+TSharedPtr<FAccountInfoEOS> FAccountInfoRegistryEOS::Find(FAccountId AccountId) const
 {
-	return StaticCastSharedPtr<FAccountInfoEOS>(Super::Find(AccountIdHandle));
+	return StaticCastSharedPtr<FAccountInfoEOS>(Super::Find(AccountId));
 }
 
 TSharedPtr<FAccountInfoEOS> FAccountInfoRegistryEOS::Find(EOS_EpicAccountId EpicAccountId) const
@@ -828,7 +828,7 @@ TOnlineAsyncOpHandle<FAuthLogout> FAuthEOSGS::Logout(FAuthLogout::Params&& Param
 	Op->Then([this](TOnlineAsyncOp<FAuthLogout>& InAsyncOp)
 	{
 		const FAuthLogout::Params& Params = InAsyncOp.GetParams();
-		TSharedPtr<FAccountInfoEOS> AccountInfoEOS = AccountInfoRegistryEOS.Find(Params.LocalUserId);
+		TSharedPtr<FAccountInfoEOS> AccountInfoEOS = AccountInfoRegistryEOS.Find(Params.LocalAccountId);
 		if (!AccountInfoEOS)
 		{
 			InAsyncOp.SetError(Errors::InvalidUser());
@@ -946,12 +946,12 @@ TOnlineAsyncOpHandle<FAuthEndVerifiedAuthSession> FAuthEOSGS::EndVerifiedAuthSes
 	return Operation->GetHandle();
 }
 
-TFuture<FAccountId> FAuthEOSGS::ResolveAccountId(const FAccountId& LocalUserId, const EOS_ProductUserId ProductUserId)
+TFuture<FAccountId> FAuthEOSGS::ResolveAccountId(const FAccountId& LocalAccountId, const EOS_ProductUserId ProductUserId)
 {
 	return MakeFulfilledPromise<FAccountId>(CreateAccountId(ProductUserId)).GetFuture();
 }
 
-TFuture<TArray<FAccountId>> FAuthEOSGS::ResolveAccountIds(const FAccountId& LocalUserId, const TArray<EOS_ProductUserId>& InProductUserIds)
+TFuture<TArray<FAccountId>> FAuthEOSGS::ResolveAccountIds(const FAccountId& LocalAccountId, const TArray<EOS_ProductUserId>& InProductUserIds)
 {
 	TArray<FAccountId> AccountIds;
 	AccountIds.Reserve(InProductUserIds.Num());
@@ -966,12 +966,12 @@ TFunction<TFuture<FAccountId>(FOnlineAsyncOp& InAsyncOp, const EOS_ProductUserId
 {
 	return [this](FOnlineAsyncOp& InAsyncOp, const EOS_ProductUserId& ProductUserId)
 	{
-		const FAccountId* LocalUserIdPtr = InAsyncOp.Data.Get<FAccountId>(TEXT("LocalUserId"));
-		if (!ensure(LocalUserIdPtr))
+		const FAccountId* LocalAccountIdPtr = InAsyncOp.Data.Get<FAccountId>(TEXT("LocalAccountId"));
+		if (!ensure(LocalAccountIdPtr))
 		{
 			return MakeFulfilledPromise<FAccountId>().GetFuture();
 		}
-		return ResolveAccountId(*LocalUserIdPtr, ProductUserId);
+		return ResolveAccountId(*LocalAccountIdPtr, ProductUserId);
 	};
 }
 
@@ -979,12 +979,12 @@ TFunction<TFuture<TArray<FAccountId>>(FOnlineAsyncOp& InAsyncOp, const TArray<EO
 {
 	return [this](FOnlineAsyncOp& InAsyncOp, const TArray<EOS_ProductUserId>& ProductUserIds)
 	{
-		const FAccountId* LocalUserIdPtr = InAsyncOp.Data.Get<FAccountId>(TEXT("LocalUserId"));
-		if (!ensure(LocalUserIdPtr))
+		const FAccountId* LocalAccountIdPtr = InAsyncOp.Data.Get<FAccountId>(TEXT("LocalAccountId"));
+		if (!ensure(LocalAccountIdPtr))
 		{
 			return MakeFulfilledPromise<TArray<FAccountId>>().GetFuture();
 		}
-		return ResolveAccountIds(*LocalUserIdPtr, ProductUserIds);
+		return ResolveAccountIds(*LocalAccountIdPtr, ProductUserIds);
 	};
 }
 
diff --git a/Engine/Plugins/Online/Experimental/OnlineServicesEOSGS/Source/Private/Online/LeaderboardsEOSGS.cpp b/Engine/Plugins/Online/Experimental/OnlineServicesEOSGS/Source/Private/Online/LeaderboardsEOSGS.cpp
--- a/Engine/Plugins/Online/Experimental/OnlineServicesEOSGS/Source/Private/Online/LeaderboardsEOSGS.cpp
+++ b/Engine/Plugins/Online/Experimental/OnlineServicesEOSGS/Source/Private/Online/LeaderboardsEOSGS.cpp
@@ -32,7 +32,7 @@ TOnlineAsyncOpHandle<FReadEntriesForUsers> FLeaderboardsEOSGS::ReadEntriesForUse
 {
 	TOnlineAsyncOpRef<FReadEntriesForUsers> Op = GetOp<FReadEntriesForUsers>(MoveTemp(Params));
 
-	if (!Services.Get<FAuthEOSGS>()->IsLoggedIn(Op->GetParams().LocalUserId))
+	if (!Services.Get<FAuthEOSGS>()->IsLoggedIn(Op->GetParams().LocalAccountId))
 	{
 		Op->SetError(Errors::InvalidUser());
 		return Op->GetHandle();
@@ -41,10 +41,10 @@ TOnlineAsyncOpHandle<FReadEntriesForUsers> FLeaderboardsEOSGS::ReadEntriesForUse
 	Op->Then([this](TOnlineAsyncOp<FReadEntriesForUsers>& InAsyncOp, TPromise<const EOS_Leaderboards_OnQueryLeaderboardUserScoresCompleteCallbackInfo*>&& Promise)
 	{
 		TArray<EOS_ProductUserId> ProductUserIds;
-		ProductUserIds.AddZeroed(InAsyncOp.GetParams().UserIds.Num());
-		for (const FAccountId& UserId : InAsyncOp.GetParams().UserIds)
+		ProductUserIds.AddZeroed(InAsyncOp.GetParams().AccountIds.Num());
+		for (const FAccountId& AccountId : InAsyncOp.GetParams().AccountIds)
 		{
-			ProductUserIds.Emplace(GetProductUserIdChecked(UserId));
+			ProductUserIds.Emplace(GetProductUserIdChecked(AccountId));
 		}
 
 		EOS_Leaderboards_UserScoresQueryStatInfo StatInfo;
@@ -66,7 +66,7 @@ TOnlineAsyncOpHandle<FReadEntriesForUsers> FLeaderboardsEOSGS::ReadEntriesForUse
 		Options.StatInfoCount = 1;
 		Options.StartTime = EOS_LEADERBOARDS_TIME_UNDEFINED;
 		Options.EndTime = EOS_LEADERBOARDS_TIME_UNDEFINED;
-		Options.LocalUserId = GetProductUserIdChecked(InAsyncOp.GetParams().LocalUserId);
+		Options.LocalUserId = GetProductUserIdChecked(InAsyncOp.GetParams().LocalAccountId);
 
 		EOS_Async(EOS_Leaderboards_QueryLeaderboardUserScores, LeaderboardsHandle, Options, MoveTemp(Promise));
 	})
@@ -84,13 +84,13 @@ TOnlineAsyncOpHandle<FReadEntriesForUsers> FLeaderboardsEOSGS::ReadEntriesForUse
 
 		FReadEntriesForUsers::Result Result;
 
-		for (const FAccountId& UserId : InAsyncOp.GetParams().UserIds)
+		for (const FAccountId& AccountId : InAsyncOp.GetParams().AccountIds)
 		{
 			EOS_Leaderboards_CopyLeaderboardUserScoreByUserIdOptions UserCopyOptions = { };
 			UserCopyOptions.ApiVersion = EOS_LEADERBOARDS_COPYLEADERBOARDUSERSCOREBYUSERID_API_LATEST;
 			static_assert(EOS_LEADERBOARDS_COPYLEADERBOARDUSERSCOREBYUSERID_API_LATEST == 1, "EOS_Leaderboards_CopyLeaderboardUserScoreByUserIdOptions updated, check new fields");
 
-			UserCopyOptions.UserId = GetProductUserIdChecked(UserId);
+			UserCopyOptions.UserId = GetProductUserIdChecked(AccountId);
 			UserCopyOptions.StatName = StatNameANSI;
 
 			EOS_Leaderboards_LeaderboardUserScore* LeaderboardUserScore = nullptr;
@@ -104,7 +104,7 @@ TOnlineAsyncOpHandle<FReadEntriesForUsers> FLeaderboardsEOSGS::ReadEntriesForUse
 
 			FLeaderboardEntry& Entry = Result.Entries.Emplace_GetRef();
 			Entry.Rank = UE_LEADERBOARD_RANK_UNKNOWN;
-			Entry.UserId = UserId;
+			Entry.AccountId = AccountId;
 			Entry.Score = LeaderboardUserScore->Score;
 
 			EOS_Leaderboards_LeaderboardUserScore_Release(LeaderboardUserScore);
@@ -120,10 +120,10 @@ TOnlineAsyncOpHandle<FReadEntriesForUsers> FLeaderboardsEOSGS::ReadEntriesForUse
 namespace Private
 {
 
-void QueryLeaderboardsEOS(EOS_HLeaderboards LeaderboardsHandle, const FAccountId& LocalUserId, const FString& BoardName, TPromise<const EOS_Leaderboards_OnQueryLeaderboardRanksCompleteCallbackInfo*>&& Promise)
+void QueryLeaderboardsEOS(EOS_HLeaderboards LeaderboardsHandle, const FAccountId& LocalAccountId, const FString& BoardName, TPromise<const EOS_Leaderboards_OnQueryLeaderboardRanksCompleteCallbackInfo*>&& Promise)
 {
 	EOS_Leaderboards_QueryLeaderboardRanksOptions Options;
-	Options.LocalUserId = GetProductUserIdChecked(LocalUserId);
+	Options.LocalUserId = GetProductUserIdChecked(LocalAccountId);
 	Options.ApiVersion = EOS_LEADERBOARDS_QUERYLEADERBOARDRANKS_API_LATEST;
 	static_assert(EOS_LEADERBOARDS_QUERYLEADERBOARDRANKS_API_LATEST == 2, "EOS_Leaderboards_QueryLeaderboardRanks updated, check new fields");
 
@@ -148,7 +148,7 @@ void ReadEntriesInRange(EOS_HLeaderboards LeaderboardsHandle, uint32 StartIndex,
 		if (CopyResult == EOS_EResult::EOS_Success)
 		{
 			FLeaderboardEntry& LeaderboardEntry = OutEntries.Emplace_GetRef();
-			LeaderboardEntry.UserId = FindAccountIdChecked(LeaderboardRecord->UserId);
+			LeaderboardEntry.AccountId = FindAccountIdChecked(LeaderboardRecord->UserId);
 			LeaderboardEntry.Rank = LeaderboardRecord->Rank;
 			LeaderboardEntry.Score = LeaderboardRecord->Score;
 
@@ -163,7 +163,7 @@ TOnlineAsyncOpHandle<FReadEntriesAroundRank> FLeaderboardsEOSGS::ReadEntriesArou
 {
 	TOnlineAsyncOpRef<FReadEntriesAroundRank> Op = GetOp<FReadEntriesAroundRank>(MoveTemp(Params));
 
-	if (!Services.Get<FAuthEOSGS>()->IsLoggedIn(Op->GetParams().LocalUserId))
+	if (!Services.Get<FAuthEOSGS>()->IsLoggedIn(Op->GetParams().LocalAccountId))
 	{
 		Op->SetError(Errors::InvalidUser());
 		return Op->GetHandle();
@@ -177,7 +177,7 @@ TOnlineAsyncOpHandle<FReadEntriesAroundRank> FLeaderboardsEOSGS::ReadEntriesArou
 
 	Op->Then([this](TOnlineAsyncOp<FReadEntriesAroundRank>& InAsyncOp, TPromise<const EOS_Leaderboards_OnQueryLeaderboardRanksCompleteCallbackInfo*>&& Promise)
 	{
-		Private::QueryLeaderboardsEOS(LeaderboardsHandle, InAsyncOp.GetParams().LocalUserId, InAsyncOp.GetParams().BoardName, MoveTemp(Promise));
+		Private::QueryLeaderboardsEOS(LeaderboardsHandle, InAsyncOp.GetParams().LocalAccountId, InAsyncOp.GetParams().BoardName, MoveTemp(Promise));
 	})
 	.Then([this](TOnlineAsyncOp<FReadEntriesAroundRank>& InAsyncOp, const EOS_Leaderboards_OnQueryLeaderboardRanksCompleteCallbackInfo* Data)
 	{
@@ -204,7 +204,7 @@ TOnlineAsyncOpHandle<FReadEntriesAroundUser> FLeaderboardsEOSGS::ReadEntriesArou
 {
 	TOnlineAsyncOpRef<FReadEntriesAroundUser> Op = GetOp<FReadEntriesAroundUser>(MoveTemp(Params));
 
-	if (!Services.Get<FAuthEOSGS>()->IsLoggedIn(Op->GetParams().LocalUserId))
+	if (!Services.Get<FAuthEOSGS>()->IsLoggedIn(Op->GetParams().LocalAccountId))
 	{
 		Op->SetError(Errors::InvalidUser());
 		return Op->GetHandle();
@@ -218,7 +218,7 @@ TOnlineAsyncOpHandle<FReadEntriesAroundUser> FLeaderboardsEOSGS::ReadEntriesArou
 
 	Op->Then([this](TOnlineAsyncOp<FReadEntriesAroundUser>& InAsyncOp, TPromise<const EOS_Leaderboards_OnQueryLeaderboardRanksCompleteCallbackInfo*>&& Promise)
 	{
-		Private::QueryLeaderboardsEOS(LeaderboardsHandle, InAsyncOp.GetParams().LocalUserId, InAsyncOp.GetParams().BoardName, MoveTemp(Promise));
+		Private::QueryLeaderboardsEOS(LeaderboardsHandle, InAsyncOp.GetParams().LocalAccountId, InAsyncOp.GetParams().BoardName, MoveTemp(Promise));
 	})
 	.Then([this](TOnlineAsyncOp<FReadEntriesAroundUser>& InAsyncOp, const EOS_Leaderboards_OnQueryLeaderboardRanksCompleteCallbackInfo* Data)
 	{
@@ -235,7 +235,7 @@ TOnlineAsyncOpHandle<FReadEntriesAroundUser> FLeaderboardsEOSGS::ReadEntriesArou
 		EOS_Leaderboards_CopyLeaderboardRecordByUserIdOptions CopyRecordByUserIdOptions;
 		CopyRecordByUserIdOptions.ApiVersion = EOS_LEADERBOARDS_COPYLEADERBOARDRECORDBYUSERID_API_LATEST;
 		static_assert(EOS_LEADERBOARDS_COPYLEADERBOARDRECORDBYUSERID_API_LATEST == 2, "EOS_Leaderboards_CopyLeaderboardRecordByUserIdOptions updated, check new fields");
-		CopyRecordByUserIdOptions.UserId = GetProductUserIdChecked(InAsyncOp.GetParams().UserId);
+		CopyRecordByUserIdOptions.UserId = GetProductUserIdChecked(InAsyncOp.GetParams().AccountId);
 		EOS_EResult CopyResult = EOS_Leaderboards_CopyLeaderboardRecordByUserId(LeaderboardsHandle, &CopyRecordByUserIdOptions, &LeaderboardRecord);
 		if (CopyResult == EOS_EResult::EOS_Success)
 		{
diff --git a/Engine/Plugins/Online/Experimental/OnlineServicesEOSGS/Source/Private/Online/LobbiesEOSGS.cpp b/Engine/Plugins/Online/Experimental/OnlineServicesEOSGS/Source/Private/Online/LobbiesEOSGS.cpp
--- a/Engine/Plugins/Online/Experimental/OnlineServicesEOSGS/Source/Private/Online/LobbiesEOSGS.cpp
+++ b/Engine/Plugins/Online/Experimental/OnlineServicesEOSGS/Source/Private/Online/LobbiesEOSGS.cpp
@@ -62,7 +62,7 @@ TOnlineAsyncOpHandle<FCreateLobby> FLobbiesEOSGS::CreateLobby(FCreateLobby::Para
 	TOnlineAsyncOpRef<FCreateLobby> Op = GetOp<FCreateLobby>(MoveTemp(InParams));
 	const FCreateLobby::Params& Params = Op->GetParams();
 
-	if (!Services.Get<FAuthEOSGS>()->IsLoggedIn(Params.LocalUserId))
+	if (!Services.Get<FAuthEOSGS>()->IsLoggedIn(Params.LocalAccountId))
 	{
 		Op->SetError(Errors::InvalidUser());
 		return Op->GetHandle();
@@ -71,12 +71,12 @@ TOnlineAsyncOpHandle<FCreateLobby> FLobbiesEOSGS::CreateLobby(FCreateLobby::Para
 	auto DestroyLobbyDuringCreate = [this](
 		TOnlineAsyncOp<FCreateLobby>& InAsyncOp,
 		const TSharedPtr<FLobbyDataEOS>& LobbyData,
-		FAccountId LocalUserId,
+		FAccountId LocalAccountId,
 		FOnlineError ErrorResult) -> TFuture<void>
 	{
 		FLobbiesDestroyLobbyImpl::Params DestroyLobbyParams;
 		DestroyLobbyParams.LobbyData = LobbyData;
-		DestroyLobbyParams.LocalUserId = LocalUserId;
+		DestroyLobbyParams.LocalAccountId = LocalAccountId;
 
 		TPromise<void> Promise;
 		auto Future = Promise.GetFuture();
@@ -109,7 +109,7 @@ TOnlineAsyncOpHandle<FCreateLobby> FLobbiesEOSGS::CreateLobby(FCreateLobby::Para
 
 		EOS_Lobby_CreateLobbyOptions CreateLobbyOptions = {};
 		CreateLobbyOptions.ApiVersion = EOS_LOBBY_CREATELOBBY_API_LATEST;
-		CreateLobbyOptions.LocalUserId = GetProductUserIdChecked(Params.LocalUserId);
+		CreateLobbyOptions.LocalUserId = GetProductUserIdChecked(Params.LocalAccountId);
 		CreateLobbyOptions.MaxLobbyMembers = Params.MaxMembers;
 		CreateLobbyOptions.PermissionLevel = TranslateJoinPolicy(ELobbyJoinPolicy::InvitationOnly);
 		CreateLobbyOptions.bPresenceEnabled = false; // todo: handle
@@ -135,7 +135,7 @@ TOnlineAsyncOpHandle<FCreateLobby> FLobbiesEOSGS::CreateLobby(FCreateLobby::Para
 		}
 
 		TDefaultErrorResultInternal<TSharedRef<FLobbyDetailsEOS>> LobbyDetailsResult =
-			FLobbyDetailsEOS::CreateFromLobbyId(LobbyPrerequisites.ToSharedRef(), Params.LocalUserId, Data->LobbyId);
+			FLobbyDetailsEOS::CreateFromLobbyId(LobbyPrerequisites.ToSharedRef(), Params.LocalAccountId, Data->LobbyId);
 		if (LobbyDetailsResult.IsError())
 		{
 			// Todo: manually call eos destroy lobby here.
@@ -145,7 +145,7 @@ TOnlineAsyncOpHandle<FCreateLobby> FLobbiesEOSGS::CreateLobby(FCreateLobby::Para
 			return MakeFulfilledPromise<TDefaultErrorResultInternal<TSharedRef<FLobbyDataEOS>>>().GetFuture();
 		}
 		
-		return LobbyDataRegistry->FindOrCreateFromLobbyDetails(InAsyncOp.GetParams().LocalUserId, LobbyDetailsResult.GetOkValue());
+		return LobbyDataRegistry->FindOrCreateFromLobbyDetails(InAsyncOp.GetParams().LocalAccountId, LobbyDetailsResult.GetOkValue());
 	})
 	.Then([this, DestroyLobbyDuringCreate](TOnlineAsyncOp<FCreateLobby>& InAsyncOp, TDefaultErrorResultInternal<TSharedRef<FLobbyDataEOS>>&& Result)
 	{
@@ -177,9 +177,9 @@ TOnlineAsyncOpHandle<FCreateLobby> FLobbiesEOSGS::CreateLobby(FCreateLobby::Para
 		// Add creator attributes if set.
 		// Todo: make this nicer.
 		const FJoinLobbyLocalUserData* UserData = Params.LocalUsers.FindByPredicate(
-		[LocalUserId = Params.LocalUserId](const FJoinLobbyLocalUserData& Data)
+		[LocalAccountId = Params.LocalAccountId](const FJoinLobbyLocalUserData& Data)
 		{
-			return LocalUserId == Data.LocalUserId;
+			return LocalAccountId == Data.LocalAccountId;
 		});
 
 		TSharedRef<FClientLobbyMemberDataChanges> LobbyOwnerAttributes = MakeShared<FClientLobbyMemberDataChanges>();
@@ -188,7 +188,7 @@ TOnlineAsyncOpHandle<FCreateLobby> FLobbiesEOSGS::CreateLobby(FCreateLobby::Para
 		// CreatingMemberData is used to update the local lobby data and for dispatching notifications once creation has completed.
 		TMap<FAccountId, TSharedRef<FClientLobbyMemberDataChanges>> CreatingMemberData;
 		CreatingMemberData.Reserve(Params.LocalUsers.Num());
-		CreatingMemberData.Add(Params.LocalUserId, LobbyOwnerAttributes);
+		CreatingMemberData.Add(Params.LocalAccountId, LobbyOwnerAttributes);
 		InAsyncOp.Data.Set<TMap<FAccountId, TSharedRef<FClientLobbyMemberDataChanges>>>(LobbyMemberChangesKeyName, MoveTemp(CreatingMemberData));
 
 		if (UserData)
@@ -197,7 +197,7 @@ TOnlineAsyncOpHandle<FCreateLobby> FLobbiesEOSGS::CreateLobby(FCreateLobby::Para
 
 			FLobbiesModifyLobbyMemberDataImpl::Params ModifyLobbyMemberDataParams;
 			ModifyLobbyMemberDataParams.LobbyData = LobbyData;
-			ModifyLobbyMemberDataParams.LocalUserId = Params.LocalUserId;
+			ModifyLobbyMemberDataParams.LocalAccountId = Params.LocalAccountId;
 			ModifyLobbyMemberDataParams.Changes = LobbyOwnerAttributes;
 
 			return ModifyLobbyMemberDataImpl(MoveTemp(ModifyLobbyMemberDataParams));
@@ -216,7 +216,7 @@ TOnlineAsyncOpHandle<FCreateLobby> FLobbiesEOSGS::CreateLobby(FCreateLobby::Para
 
 		if (ModifyLobbyOwnerResult.IsError())
 		{
-			return DestroyLobbyDuringCreate(InAsyncOp, LobbyData, Params.LocalUserId, MoveTemp(ModifyLobbyOwnerResult.GetErrorValue()));
+			return DestroyLobbyDuringCreate(InAsyncOp, LobbyData, Params.LocalAccountId, MoveTemp(ModifyLobbyOwnerResult.GetErrorValue()));
 		}
 		else
 		{
@@ -246,7 +246,7 @@ TOnlineAsyncOpHandle<FCreateLobby> FLobbiesEOSGS::CreateLobby(FCreateLobby::Para
 		// Set lobby privacy setting to the user provided value.
 		FLobbiesModifyLobbyDataImpl::Params ModifyLobbyDataParams;
 		ModifyLobbyDataParams.LobbyData = LobbyData;
-		ModifyLobbyDataParams.LocalUserId = Params.LocalUserId;
+		ModifyLobbyDataParams.LocalAccountId = Params.LocalAccountId;
 		ModifyLobbyDataParams.Changes = LobbyChanges;
 
 		return ModifyLobbyDataImpl(MoveTemp(ModifyLobbyDataParams));
@@ -260,7 +260,7 @@ TOnlineAsyncOpHandle<FCreateLobby> FLobbiesEOSGS::CreateLobby(FCreateLobby::Para
 
 		if (ModifyLobbyResult.IsError())
 		{
-			return DestroyLobbyDuringCreate(InAsyncOp, LobbyData, Params.LocalUserId, MoveTemp(ModifyLobbyResult.GetErrorValue()));
+			return DestroyLobbyDuringCreate(InAsyncOp, LobbyData, Params.LocalAccountId, MoveTemp(ModifyLobbyResult.GetErrorValue()));
 		}
 		else
 		{
@@ -279,7 +279,7 @@ TOnlineAsyncOpHandle<FCreateLobby> FLobbiesEOSGS::CreateLobby(FCreateLobby::Para
 		// Mark the lobby active for each member.
 		for (const FJoinLobbyLocalUserData& LocalUser : Params.LocalUsers)
 		{
-			AddActiveLobby(LocalUser.LocalUserId, LobbyData);
+			AddActiveLobby(LocalUser.LocalAccountId, LobbyData);
 		}
 
 		// Add member changes to the lobby changes object.
@@ -304,14 +304,14 @@ TOnlineAsyncOpHandle<FFindLobbies> FLobbiesEOSGS::FindLobbies(FFindLobbies::Para
 	TOnlineAsyncOpRef<FFindLobbies> Op = GetOp<FFindLobbies>(MoveTemp(InParams));
 	const FFindLobbies::Params& Params = Op->GetParams();
 
-	if (!Services.Get<FAuthEOSGS>()->IsLoggedIn(Params.LocalUserId))
+	if (!Services.Get<FAuthEOSGS>()->IsLoggedIn(Params.LocalAccountId))
 	{
 		Op->SetError(Errors::InvalidUser());
 		return Op->GetHandle();
 	}
 
 	// Invalidate previous search results.
-	ActiveSearchResults.Remove(Params.LocalUserId);
+	ActiveSearchResults.Remove(Params.LocalAccountId);
 
 	Op->Then([this](TOnlineAsyncOp<FFindLobbies>& InAsyncOp)
 	{
@@ -326,7 +326,7 @@ TOnlineAsyncOpHandle<FFindLobbies> FLobbiesEOSGS::FindLobbies(FFindLobbies::Para
 		}
 		else
 		{
-			ActiveSearchResults.Add(InAsyncOp.GetParams().LocalUserId, Result.GetOkValue());
+			ActiveSearchResults.Add(InAsyncOp.GetParams().LocalAccountId, Result.GetOkValue());
 			InAsyncOp.SetResult(FFindLobbies::Result{Result.GetOkValue()->GetLobbyResults()});
 		}
 	})
@@ -340,7 +340,7 @@ TOnlineAsyncOpHandle<FJoinLobby> FLobbiesEOSGS::JoinLobby(FJoinLobby::Params&& I
 	TOnlineAsyncOpRef<FJoinLobby> Op = GetOp<FJoinLobby>(MoveTemp(InParams));
 	const FJoinLobby::Params& Params = Op->GetParams();
 
-	if (!Services.Get<FAuthEOSGS>()->IsLoggedIn(Params.LocalUserId))
+	if (!Services.Get<FAuthEOSGS>()->IsLoggedIn(Params.LocalAccountId))
 	{
 		Op->SetError(Errors::InvalidUser());
 		return Op->GetHandle();
@@ -361,7 +361,7 @@ TOnlineAsyncOpHandle<FJoinLobby> FLobbiesEOSGS::JoinLobby(FJoinLobby::Params&& I
 
 		FLobbiesJoinLobbyImpl::Params JoinParams;
 		JoinParams.LobbyData = LobbyData;
-		JoinParams.LocalUserId = Params.LocalUserId;
+		JoinParams.LocalAccountId = Params.LocalAccountId;
 		JoinParams.LocalName = Params.LocalName;
 		JoinParams.LocalUsers = Params.LocalUsers;
 		return JoinLobbyImpl(MoveTemp(JoinParams));
@@ -385,7 +385,7 @@ TOnlineAsyncOpHandle<FJoinLobby> FLobbiesEOSGS::JoinLobby(FJoinLobby::Params&& I
 		// Mark the lobby active for each member.
 		for (const FJoinLobbyLocalUserData& LocalUser : Params.LocalUsers)
 		{
-			AddActiveLobby(LocalUser.LocalUserId, LobbyData.ToSharedRef());
+			AddActiveLobby(LocalUser.LocalAccountId, LobbyData.ToSharedRef());
 		}
 
 		// todo: figure out a butter way to handle attribute parameters.
@@ -396,7 +396,7 @@ TOnlineAsyncOpHandle<FJoinLobby> FLobbiesEOSGS::JoinLobby(FJoinLobby::Params&& I
 		{
 			TSharedRef<FClientLobbyMemberDataChanges> MemberDataChanges = MakeShared<FClientLobbyMemberDataChanges>();
 			MemberDataChanges->MutatedAttributes = LocalUser.Attributes;
-			LobbyChanges.MutatedMembers.Add(LocalUser.LocalUserId, MoveTemp(MemberDataChanges));
+			LobbyChanges.MutatedMembers.Add(LocalUser.LocalAccountId, MoveTemp(MemberDataChanges));
 		}
 		LobbyData->GetClientLobbyData()->ApplyLobbyUpdateFromLocalChanges(MoveTemp(LobbyChanges), LobbyEvents);
 
@@ -412,7 +412,7 @@ TOnlineAsyncOpHandle<FLeaveLobby> FLobbiesEOSGS::LeaveLobby(FLeaveLobby::Params&
 	TOnlineAsyncOpRef<FLeaveLobby> Op = GetOp<FLeaveLobby>(MoveTemp(InParams));
 	const FLeaveLobby::Params& Params = Op->GetParams();
 
-	if (!Services.Get<FAuthEOSGS>()->IsLoggedIn(Params.LocalUserId))
+	if (!Services.Get<FAuthEOSGS>()->IsLoggedIn(Params.LocalAccountId))
 	{
 		Op->SetError(Errors::InvalidUser());
 		return Op->GetHandle();
@@ -433,7 +433,7 @@ TOnlineAsyncOpHandle<FLeaveLobby> FLobbiesEOSGS::LeaveLobby(FLeaveLobby::Params&
 
 		FLobbiesLeaveLobbyImpl::Params LeaveParams;
 		LeaveParams.LobbyData = LobbyData;
-		LeaveParams.LocalUserId = Params.LocalUserId;
+		LeaveParams.LocalAccountId = Params.LocalAccountId;
 		return LeaveLobbyImpl(MoveTemp(LeaveParams));
 	})
 	.Then([this, LobbyData](TOnlineAsyncOp<FLeaveLobby>& InAsyncOp, TDefaultErrorResult<FLobbiesLeaveLobbyImpl> Result)
@@ -442,7 +442,7 @@ TOnlineAsyncOpHandle<FLeaveLobby> FLobbiesEOSGS::LeaveLobby(FLeaveLobby::Params&
 
 		// Remove the user from the local lobby data and dispatch notifications.
 		FClientLobbyDataChanges LobbyChanges;
-		LobbyChanges.LeavingMembers.Add(Params.LocalUserId, ELobbyMemberLeaveReason::Left);
+		LobbyChanges.LeavingMembers.Add(Params.LocalAccountId, ELobbyMemberLeaveReason::Left);
 		FApplyLobbyUpdateResult ApplyResult = LobbyData->GetClientLobbyData()->ApplyLobbyUpdateFromLocalChanges(MoveTemp(LobbyChanges), LobbyEvents);
 
 		// Remove the lobby from the active list for the user.
@@ -472,7 +472,7 @@ TOnlineAsyncOpHandle<FInviteLobbyMember> FLobbiesEOSGS::InviteLobbyMember(FInvit
 	TOnlineAsyncOpRef<FInviteLobbyMember> Op = GetOp<FInviteLobbyMember>(MoveTemp(InParams));
 	const FInviteLobbyMember::Params& Params = Op->GetParams();
 
-	if (!Services.Get<FAuthEOSGS>()->IsLoggedIn(Params.LocalUserId))
+	if (!Services.Get<FAuthEOSGS>()->IsLoggedIn(Params.LocalAccountId))
 	{
 		Op->SetError(Errors::InvalidUser());
 		return Op->GetHandle();
@@ -484,8 +484,8 @@ TOnlineAsyncOpHandle<FInviteLobbyMember> FLobbiesEOSGS::InviteLobbyMember(FInvit
 
 		FLobbiesInviteLobbyMemberImpl::Params InviteParams;
 		InviteParams.LobbyData = LobbyDataRegistry->Find(Params.LobbyId);
-		InviteParams.LocalUserId = Params.LocalUserId;
-		InviteParams.TargetUserId = Params.TargetUserId;
+		InviteParams.LocalAccountId = Params.LocalAccountId;
+		InviteParams.TargetAccountId = Params.TargetAccountId;
 		return InviteLobbyMemberImpl(MoveTemp(InviteParams));
 	})
 	.Then([this](TOnlineAsyncOp<FInviteLobbyMember>& InAsyncOp, TDefaultErrorResult<FLobbiesInviteLobbyMemberImpl> Result)
@@ -510,7 +510,7 @@ TOnlineAsyncOpHandle<FDeclineLobbyInvitation> FLobbiesEOSGS::DeclineLobbyInvitat
 	TOnlineAsyncOpRef<FDeclineLobbyInvitation> Op = GetOp<FDeclineLobbyInvitation>(MoveTemp(InParams));
 	const FDeclineLobbyInvitation::Params& Params = Op->GetParams();
 
-	if (!Services.Get<FAuthEOSGS>()->IsLoggedIn(Params.LocalUserId))
+	if (!Services.Get<FAuthEOSGS>()->IsLoggedIn(Params.LocalAccountId))
 	{
 		Op->SetError(Errors::InvalidUser());
 		return Op->GetHandle();
@@ -521,7 +521,7 @@ TOnlineAsyncOpHandle<FDeclineLobbyInvitation> FLobbiesEOSGS::DeclineLobbyInvitat
 		const FDeclineLobbyInvitation::Params& Params = InAsyncOp.GetParams();
 
 		FLobbiesDeclineLobbyInvitationImpl::Params DeclineParams;
-		DeclineParams.LocalUserId = Params.LocalUserId;
+		DeclineParams.LocalAccountId = Params.LocalAccountId;
 		DeclineParams.LobbyId = Params.LobbyId;
 		return DeclineLobbyInvitationImpl(MoveTemp(DeclineParams));
 	})
@@ -547,7 +547,7 @@ TOnlineAsyncOpHandle<FKickLobbyMember> FLobbiesEOSGS::KickLobbyMember(FKickLobby
 	TOnlineAsyncOpRef<FKickLobbyMember> Op = GetOp<FKickLobbyMember>(MoveTemp(InParams));
 	const FKickLobbyMember::Params& Params = Op->GetParams();
 
-	if (!Services.Get<FAuthEOSGS>()->IsLoggedIn(Params.LocalUserId))
+	if (!Services.Get<FAuthEOSGS>()->IsLoggedIn(Params.LocalAccountId))
 	{
 		Op->SetError(Errors::InvalidUser());
 		return Op->GetHandle();
@@ -568,8 +568,8 @@ TOnlineAsyncOpHandle<FKickLobbyMember> FLobbiesEOSGS::KickLobbyMember(FKickLobby
 
 		FLobbiesKickLobbyMemberImpl::Params KickParams;
 		KickParams.LobbyData = LobbyData;
-		KickParams.LocalUserId = Params.LocalUserId;
-		KickParams.TargetUserId = Params.TargetUserId;
+		KickParams.LocalAccountId = Params.LocalAccountId;
+		KickParams.TargetAccountId = Params.TargetAccountId;
 		return KickLobbyMemberImpl(MoveTemp(KickParams));
 	})
 	.Then([this](TOnlineAsyncOp<FKickLobbyMember>& InAsyncOp, TDefaultErrorResult<FLobbiesKickLobbyMemberImpl> Result)
@@ -587,7 +587,7 @@ TOnlineAsyncOpHandle<FKickLobbyMember> FLobbiesEOSGS::KickLobbyMember(FKickLobby
 
 		// Update local cache and fire events.
 		FClientLobbyDataChanges LobbyChanges;
-		LobbyChanges.LeavingMembers.Add(Params.TargetUserId, ELobbyMemberLeaveReason::Kicked);
+		LobbyChanges.LeavingMembers.Add(Params.TargetAccountId, ELobbyMemberLeaveReason::Kicked);
 		LobbyData->GetClientLobbyData()->ApplyLobbyUpdateFromLocalChanges(MoveTemp(LobbyChanges), LobbyEvents);
 		InAsyncOp.SetResult(FKickLobbyMember::Result{});
 	})
@@ -601,7 +601,7 @@ TOnlineAsyncOpHandle<FPromoteLobbyMember> FLobbiesEOSGS::PromoteLobbyMember(FPro
 	TOnlineAsyncOpRef<FPromoteLobbyMember> Op = GetOp<FPromoteLobbyMember>(MoveTemp(InParams));
 	const FPromoteLobbyMember::Params& Params = Op->GetParams();
 
-	if (!Services.Get<FAuthEOSGS>()->IsLoggedIn(Params.LocalUserId))
+	if (!Services.Get<FAuthEOSGS>()->IsLoggedIn(Params.LocalAccountId))
 	{
 		Op->SetError(Errors::InvalidUser());
 		return Op->GetHandle();
@@ -622,8 +622,8 @@ TOnlineAsyncOpHandle<FPromoteLobbyMember> FLobbiesEOSGS::PromoteLobbyMember(FPro
 
 		FLobbiesPromoteLobbyMemberImpl::Params PromoteParams;
 		PromoteParams.LobbyData = LobbyData;
-		PromoteParams.LocalUserId = Params.LocalUserId;
-		PromoteParams.TargetUserId = Params.TargetUserId;
+		PromoteParams.LocalAccountId = Params.LocalAccountId;
+		PromoteParams.TargetAccountId = Params.TargetAccountId;
 		return PromoteLobbyMemberImpl(MoveTemp(PromoteParams));
 	})
 	.Then([this](TOnlineAsyncOp<FPromoteLobbyMember>& InAsyncOp, TDefaultErrorResult<FLobbiesPromoteLobbyMemberImpl> Result)
@@ -641,7 +641,7 @@ TOnlineAsyncOpHandle<FPromoteLobbyMember> FLobbiesEOSGS::PromoteLobbyMember(FPro
 
 		// Update local cache and fire events.
 		FClientLobbyDataChanges LobbyChanges;
-		LobbyChanges.OwnerAccountId = Params.TargetUserId;
+		LobbyChanges.OwnerAccountId = Params.TargetAccountId;
 		LobbyData->GetClientLobbyData()->ApplyLobbyUpdateFromLocalChanges(MoveTemp(LobbyChanges), LobbyEvents);
 		InAsyncOp.SetResult(FPromoteLobbyMember::Result{});
 	})
@@ -655,7 +655,7 @@ TOnlineAsyncOpHandle<FModifyLobbyJoinPolicy> FLobbiesEOSGS::ModifyLobbyJoinPolic
 	TOnlineAsyncOpRef<FModifyLobbyJoinPolicy> Op = GetOp<FModifyLobbyJoinPolicy>(MoveTemp(InParams));
 	const FModifyLobbyJoinPolicy::Params& Params = Op->GetParams();
 
-	if (!Services.Get<FAuthEOSGS>()->IsLoggedIn(Params.LocalUserId))
+	if (!Services.Get<FAuthEOSGS>()->IsLoggedIn(Params.LocalAccountId))
 	{
 		Op->SetError(Errors::InvalidUser());
 		return Op->GetHandle();
@@ -680,7 +680,7 @@ TOnlineAsyncOpHandle<FModifyLobbyJoinPolicy> FLobbiesEOSGS::ModifyLobbyJoinPolic
 
 		FLobbiesModifyLobbyDataImpl::Params ModifyLobbyDataParams;
 		ModifyLobbyDataParams.LobbyData = LobbyData;
-		ModifyLobbyDataParams.LocalUserId = Params.LocalUserId;
+		ModifyLobbyDataParams.LocalAccountId = Params.LocalAccountId;
 		ModifyLobbyDataParams.Changes = LobbyChanges;
 		return ModifyLobbyDataImpl(MoveTemp(ModifyLobbyDataParams));
 	})
@@ -711,7 +711,7 @@ TOnlineAsyncOpHandle<FModifyLobbyAttributes> FLobbiesEOSGS::ModifyLobbyAttribute
 	TOnlineAsyncOpRef<FModifyLobbyAttributes> Op = GetOp<FModifyLobbyAttributes>(MoveTemp(InParams));
 	const FModifyLobbyAttributes::Params& Params = Op->GetParams();
 
-	if (!Services.Get<FAuthEOSGS>()->IsLoggedIn(Params.LocalUserId))
+	if (!Services.Get<FAuthEOSGS>()->IsLoggedIn(Params.LocalAccountId))
 	{
 		Op->SetError(Errors::InvalidUser());
 		return Op->GetHandle();
@@ -737,7 +737,7 @@ TOnlineAsyncOpHandle<FModifyLobbyAttributes> FLobbiesEOSGS::ModifyLobbyAttribute
 
 		FLobbiesModifyLobbyDataImpl::Params ModifyLobbyDataParams;
 		ModifyLobbyDataParams.LobbyData = LobbyDataRegistry->Find(Params.LobbyId);
-		ModifyLobbyDataParams.LocalUserId = Params.LocalUserId;
+		ModifyLobbyDataParams.LocalAccountId = Params.LocalAccountId;
 		ModifyLobbyDataParams.Changes = LobbyChanges;
 		return ModifyLobbyDataImpl(MoveTemp(ModifyLobbyDataParams));
 	})
@@ -768,7 +768,7 @@ TOnlineAsyncOpHandle<FModifyLobbyMemberAttributes> FLobbiesEOSGS::ModifyLobbyMem
 	TOnlineAsyncOpRef<FModifyLobbyMemberAttributes> Op = GetOp<FModifyLobbyMemberAttributes>(MoveTemp(InParams));
 	const FModifyLobbyMemberAttributes::Params& Params = Op->GetParams();
 
-	if (!Services.Get<FAuthEOSGS>()->IsLoggedIn(Params.LocalUserId))
+	if (!Services.Get<FAuthEOSGS>()->IsLoggedIn(Params.LocalAccountId))
 	{
 		Op->SetError(Errors::InvalidUser());
 		return Op->GetHandle();
@@ -794,7 +794,7 @@ TOnlineAsyncOpHandle<FModifyLobbyMemberAttributes> FLobbiesEOSGS::ModifyLobbyMem
 
 		FLobbiesModifyLobbyMemberDataImpl::Params ModifyLobbyMemberDataParams;
 		ModifyLobbyMemberDataParams.LobbyData = LobbyData;
-		ModifyLobbyMemberDataParams.LocalUserId = Params.LocalUserId;
+		ModifyLobbyMemberDataParams.LocalAccountId = Params.LocalAccountId;
 		ModifyLobbyMemberDataParams.Changes = LobbyMemberChanges;
 		return ModifyLobbyMemberDataImpl(MoveTemp(ModifyLobbyMemberDataParams));
 	})
@@ -814,7 +814,7 @@ TOnlineAsyncOpHandle<FModifyLobbyMemberAttributes> FLobbiesEOSGS::ModifyLobbyMem
 
 		// Update local cache and fire events.
 		FClientLobbyDataChanges LobbyChanges;
-		LobbyChanges.MutatedMembers.Add(Params.LocalUserId, LobbyMemberChanges);
+		LobbyChanges.MutatedMembers.Add(Params.LocalAccountId, LobbyMemberChanges);
 		LobbyData->GetClientLobbyData()->ApplyLobbyUpdateFromLocalChanges(MoveTemp(LobbyChanges), LobbyEvents);
 		InAsyncOp.SetResult(FModifyLobbyMemberAttributes::Result{});
 	})
@@ -825,7 +825,7 @@ TOnlineAsyncOpHandle<FModifyLobbyMemberAttributes> FLobbiesEOSGS::ModifyLobbyMem
 
 TOnlineResult<FGetJoinedLobbies> FLobbiesEOSGS::GetJoinedLobbies(FGetJoinedLobbies::Params&& Params)
 {
-	if (const TSet<TSharedRef<FLobbyDataEOS>>* Lobbies = ActiveLobbies.Find(Params.LocalUserId))
+	if (const TSet<TSharedRef<FLobbyDataEOS>>* Lobbies = ActiveLobbies.Find(Params.LocalAccountId))
 	{
 		FGetJoinedLobbies::Result Result;
 		Result.Lobbies.Reserve(Lobbies->Num());
@@ -944,10 +944,10 @@ void FLobbiesEOSGS::HandleLobbyInviteReceived(const EOS_Lobby_LobbyInviteReceive
 #endif
 
 	// Todo: Queue this like an operation.
-	const FAccountId LocalUserId = FindAccountId(Data->LocalUserId);
-	if (LocalUserId.IsValid())
+	const FAccountId LocalAccountId = FindAccountId(Data->LocalUserId);
+	if (LocalAccountId.IsValid())
 	{
-		FLobbyInviteDataEOS::CreateFromInviteId(LobbyPrerequisites.ToSharedRef(), LobbyDataRegistry.ToSharedRef(), LocalUserId, Data->InviteId, Data->TargetUserId)
+		FLobbyInviteDataEOS::CreateFromInviteId(LobbyPrerequisites.ToSharedRef(), LobbyDataRegistry.ToSharedRef(), LocalAccountId, Data->InviteId, Data->TargetUserId)
 		.Then([this](TFuture<TDefaultErrorResultInternal<TSharedRef<FLobbyInviteDataEOS>>>&& Future)
 		{
 			if (Future.Get().IsError())
@@ -972,10 +972,10 @@ void FLobbiesEOSGS::HandleLobbyInviteReceived(const EOS_Lobby_LobbyInviteReceive
 void FLobbiesEOSGS::HandleLobbyInviteAccepted(const EOS_Lobby_LobbyInviteAcceptedCallbackInfo* Data)
 {
 	// Todo: Queue this like an operation.
-	const FAccountId LocalUserId = FindAccountId(Data->LocalUserId);
-	if (LocalUserId.IsValid())
+	const FAccountId LocalAccountId = FindAccountId(Data->LocalUserId);
+	if (LocalAccountId.IsValid())
 	{
-		FLobbyInviteDataEOS::CreateFromInviteId(LobbyPrerequisites.ToSharedRef(), LobbyDataRegistry.ToSharedRef(), LocalUserId, Data->InviteId, Data->TargetUserId)
+		FLobbyInviteDataEOS::CreateFromInviteId(LobbyPrerequisites.ToSharedRef(), LobbyDataRegistry.ToSharedRef(), LocalAccountId, Data->InviteId, Data->TargetUserId)
 			.Then([this](TFuture<TDefaultErrorResultInternal<TSharedRef<FLobbyInviteDataEOS>>>&& Future)
 				{
 					if (Future.Get().IsError())
@@ -1007,22 +1007,22 @@ void FLobbiesEOSGS::HandleLobbyInviteAccepted(const EOS_Lobby_LobbyInviteAccepte
 void FLobbiesEOSGS::HandleJoinLobbyAccepted(const EOS_Lobby_JoinLobbyAcceptedCallbackInfo* Data)
 {
 	// Todo: Queue this like an operation.
-	const FAccountId LocalUserId = FindAccountId(Data->LocalUserId);
-	if (LocalUserId.IsValid())
+	const FAccountId LocalAccountId = FindAccountId(Data->LocalUserId);
+	if (LocalAccountId.IsValid())
 	{
 		TDefaultErrorResultInternal<TSharedRef<FLobbyDetailsEOS>> LobbyDetailsResult = 
-			FLobbyDetailsEOS::CreateFromUiEventId(LobbyPrerequisites.ToSharedRef(), LocalUserId, Data->UiEventId);
+			FLobbyDetailsEOS::CreateFromUiEventId(LobbyPrerequisites.ToSharedRef(), LocalAccountId, Data->UiEventId);
 		if (LobbyDetailsResult.IsError())
 		{
 			// Todo: Log / queue a manual fetch of invitations.
 			UE_LOG(LogTemp, Warning, TEXT("[FLobbiesEOSGS::HandleJoinLobbyAccepted] Failed to receive invite. Error: %s"),
 				*LobbyDetailsResult.GetErrorValue().GetLogString());
 		}
-		LobbyDataRegistry->FindOrCreateFromLobbyDetails(LocalUserId, LobbyDetailsResult.GetOkValue()).Then([this, LocalUserId](TFuture<TDefaultErrorResultInternal<TSharedRef<FLobbyDataEOS>>> && Result)
+		LobbyDataRegistry->FindOrCreateFromLobbyDetails(LocalAccountId, LobbyDetailsResult.GetOkValue()).Then([this, LocalAccountId](TFuture<TDefaultErrorResultInternal<TSharedRef<FLobbyDataEOS>>> && Result)
 		{
 			LobbyEvents.OnUILobbyJoinRequested.Broadcast(
 				FUILobbyJoinRequested{
-					LocalUserId,
+					LocalAccountId,
 					TResult<TSharedRef<const FLobby>, FOnlineError>(Result.Get().GetOkValue()->GetClientLobbyData()->GetPublicDataPtr()),
 					EUILobbyJoinRequestedSource::Unspecified
 				});
@@ -1154,16 +1154,16 @@ void FLobbiesEOSGS::UnregisterHandlers()
 	OnJoinLobbyAcceptedEOSEventRegistration = nullptr;
 }
 
-void FLobbiesEOSGS::AddActiveLobby(FAccountId LocalUserId, const TSharedRef<FLobbyDataEOS>& LobbyData)
+void FLobbiesEOSGS::AddActiveLobby(FAccountId LocalAccountId, const TSharedRef<FLobbyDataEOS>& LobbyData)
 {
 	// Add bookkeeping for the user.
-	ActiveLobbies.FindOrAdd(LocalUserId).Add(LobbyData);
+	ActiveLobbies.FindOrAdd(LocalAccountId).Add(LobbyData);
 }
 
-void FLobbiesEOSGS::RemoveActiveLobby(FAccountId LocalUserId, const TSharedRef<FLobbyDataEOS>& LobbyData)
+void FLobbiesEOSGS::RemoveActiveLobby(FAccountId LocalAccountId, const TSharedRef<FLobbyDataEOS>& LobbyData)
 {
 	// Remove bookkeeping for the local user.
-	if (TSet<TSharedRef<FLobbyDataEOS>>* Lobbies = ActiveLobbies.Find(LocalUserId))
+	if (TSet<TSharedRef<FLobbyDataEOS>>* Lobbies = ActiveLobbies.Find(LocalAccountId))
 	{
 		Lobbies->Remove(LobbyData);
 	}
@@ -1208,7 +1208,7 @@ TSharedPtr<FLobbyInviteDataEOS> FLobbiesEOSGS::GetActiveInvite(FAccountId Target
 TFuture<TDefaultErrorResult<FLobbiesJoinLobbyImpl>> FLobbiesEOSGS::JoinLobbyImpl(FLobbiesJoinLobbyImpl::Params&& Params)
 {
 	// Check prerequisites.
-	if (!Services.Get<FAuthEOSGS>()->IsLoggedIn(Params.LocalUserId))
+	if (!Services.Get<FAuthEOSGS>()->IsLoggedIn(Params.LocalAccountId))
 	{
 		return MakeFulfilledPromise<TDefaultErrorResult<FLobbiesJoinLobbyImpl>>(Errors::InvalidUser()).GetFuture();
 	}
@@ -1221,7 +1221,7 @@ TFuture<TDefaultErrorResult<FLobbiesJoinLobbyImpl>> FLobbiesEOSGS::JoinLobbyImpl
 	// Check whether any of the local users is already in the target lobby.
 	for (const FJoinLobbyLocalUserData& JoinData : Params.LocalUsers)
 	{
-		if (const TSet<TSharedRef<FLobbyDataEOS>>* UserActiveLobbies = ActiveLobbies.Find(JoinData.LocalUserId))
+		if (const TSet<TSharedRef<FLobbyDataEOS>>* UserActiveLobbies = ActiveLobbies.Find(JoinData.LocalAccountId))
 		{
 			if (UserActiveLobbies->Find(Params.LobbyData.ToSharedRef()))
 			{
@@ -1235,7 +1235,7 @@ TFuture<TDefaultErrorResult<FLobbiesJoinLobbyImpl>> FLobbiesEOSGS::JoinLobbyImpl
 	{
 		FLobbiesJoinLobbyMemberImpl::Params JoinLobbyMemberParams;
 		JoinLobbyMemberParams.LobbyData = Params.LobbyData.ToSharedRef();
-		JoinLobbyMemberParams.LocalUserId = UserData.LocalUserId;
+		JoinLobbyMemberParams.LocalAccountId = UserData.LocalAccountId;
 		JoinLobbyMemberParams.Attributes = MoveTemp(UserData.Attributes);
 
 		TSharedRef<TPromise<TDefaultErrorResult<FLobbiesJoinLobbyMemberImpl>>> JoinMemberPromise = MakeShared<TPromise<TDefaultErrorResult<FLobbiesJoinLobbyMemberImpl>>>();
@@ -1276,13 +1276,13 @@ TFuture<TDefaultErrorResult<FLobbiesJoinLobbyImpl>> FLobbiesEOSGS::JoinLobbyImpl
 			TArray<TFuture<TDefaultErrorResult<FLobbiesLeaveLobbyImpl>>> PendingMemberExits;
 			for (int32 MemberIndex = 0; MemberIndex < Params.LocalUsers.Num(); ++MemberIndex)
 			{
-				const FAccountId MemberId = Params.LocalUsers[MemberIndex].LocalUserId;
+				const FAccountId MemberId = Params.LocalUsers[MemberIndex].LocalAccountId;
 
 				if (Results.Get()[MemberIndex].IsError())
 				{
 					FLobbiesLeaveLobbyImpl::Params LeaveLobbyParams;
 					LeaveLobbyParams.LobbyData = Params.LobbyData.ToSharedRef();
-					LeaveLobbyParams.LocalUserId = MemberId;
+					LeaveLobbyParams.LocalAccountId = MemberId;
 
 					TPromise<TDefaultErrorResult<FLobbiesLeaveLobbyImpl>> LeaveMemberPromise;
 					PendingMemberExits.Emplace(LeaveMemberPromise.GetFuture());
@@ -1321,7 +1321,7 @@ TOnlineAsyncOpHandle<FLobbiesJoinLobbyMemberImpl> FLobbiesEOSGS::JoinLobbyMember
 	const FLobbiesJoinLobbyMemberImpl::Params& Params = Op->GetParams();
 
 	// Setup lobby details - Prefer UI event before invitation before search result.
-	TSharedPtr<FLobbyDetailsEOS> LobbyDetails = Params.LobbyData->GetUserLobbyDetails(Params.LocalUserId);
+	TSharedPtr<FLobbyDetailsEOS> LobbyDetails = Params.LobbyData->GetUserLobbyDetails(Params.LocalAccountId);
 	if (!LobbyDetails)
 	{
 		// Todo: Check whether another local member can invite the user.
@@ -1342,7 +1342,7 @@ TOnlineAsyncOpHandle<FLobbiesJoinLobbyMemberImpl> FLobbiesEOSGS::JoinLobbyMember
 		EOS_Lobby_JoinLobbyOptions JoinLobbyOptions = {};
 		JoinLobbyOptions.ApiVersion = EOS_LOBBY_JOINLOBBY_API_LATEST;
 		JoinLobbyOptions.LobbyDetailsHandle = LobbyDetails->GetEOSHandle();
-		JoinLobbyOptions.LocalUserId = GetProductUserIdChecked(Params.LocalUserId);
+		JoinLobbyOptions.LocalUserId = GetProductUserIdChecked(Params.LocalAccountId);
 		JoinLobbyOptions.bPresenceEnabled = false; // todo
 		JoinLobbyOptions.LocalRTCOptions = nullptr;
 		static_assert(EOS_LOBBY_JOINLOBBY_API_LATEST == 3, "EOS_Lobby_JoinLobbyOptions updated, check new fields");
@@ -1362,7 +1362,7 @@ TOnlineAsyncOpHandle<FLobbiesJoinLobbyMemberImpl> FLobbiesEOSGS::JoinLobbyMember
 		const FLobbiesJoinLobbyMemberImpl::Params& Params = InAsyncOp.GetParams();
 		FLobbiesModifyLobbyMemberDataImpl::Params ModifyParams;
 		ModifyParams.LobbyData = Params.LobbyData;
-		ModifyParams.LocalUserId = Params.LocalUserId;
+		ModifyParams.LocalAccountId = Params.LocalAccountId;
 		ModifyParams.Changes = MakeShared<FClientLobbyMemberDataChanges>();
 		ModifyParams.Changes->MutatedAttributes = Params.Attributes;
 		return ModifyLobbyMemberDataImpl(MoveTemp(ModifyParams));
@@ -1376,7 +1376,7 @@ TOnlineAsyncOpHandle<FLobbiesJoinLobbyMemberImpl> FLobbiesEOSGS::JoinLobbyMember
 			// Failed to set attributes - leave the lobby.
 			FLobbiesLeaveLobbyImpl::Params LeaveLobbyParams;
 			LeaveLobbyParams.LobbyData = Params.LobbyData;
-			LeaveLobbyParams.LocalUserId = Params.LocalUserId;
+			LeaveLobbyParams.LocalAccountId = Params.LocalAccountId;
 
 			TPromise<void> Promise;
 			TFuture<void> Future = Promise.GetFuture();
@@ -1405,7 +1405,7 @@ TOnlineAsyncOpHandle<FLobbiesJoinLobbyMemberImpl> FLobbiesEOSGS::JoinLobbyMember
 	{
 		InAsyncOp.SetResult(FLobbiesJoinLobbyMemberImpl::Result{});
 	})
-	.Enqueue(GetSerialQueue(InParams.LocalUserId));
+	.Enqueue(GetSerialQueue(InParams.LocalAccountId));
 
 	return Op->GetHandle();
 }
@@ -1413,7 +1413,7 @@ TOnlineAsyncOpHandle<FLobbiesJoinLobbyMemberImpl> FLobbiesEOSGS::JoinLobbyMember
 TFuture<TDefaultErrorResult<FLobbiesLeaveLobbyImpl>> FLobbiesEOSGS::LeaveLobbyImpl(FLobbiesLeaveLobbyImpl::Params&& Params)
 {
 	// Check prerequisites.
-	if (!Services.Get<FAuthEOSGS>()->IsLoggedIn(Params.LocalUserId))
+	if (!Services.Get<FAuthEOSGS>()->IsLoggedIn(Params.LocalAccountId))
 	{
 		return MakeFulfilledPromise<TDefaultErrorResult<FLobbiesLeaveLobbyImpl>>(Errors::InvalidUser()).GetFuture();
 	}
@@ -1425,7 +1425,7 @@ TFuture<TDefaultErrorResult<FLobbiesLeaveLobbyImpl>> FLobbiesEOSGS::LeaveLobbyIm
 
 	EOS_Lobby_LeaveLobbyOptions LeaveLobbyOptions = {};
 	LeaveLobbyOptions.ApiVersion = EOS_LOBBY_LEAVELOBBY_API_LATEST;
-	LeaveLobbyOptions.LocalUserId = GetProductUserIdChecked(Params.LocalUserId);
+	LeaveLobbyOptions.LocalUserId = GetProductUserIdChecked(Params.LocalAccountId);
 	LeaveLobbyOptions.LobbyId = Params.LobbyData->GetLobbyIdEOS();
 	static_assert(EOS_LOBBY_LEAVELOBBY_API_LATEST == 1, "EOS_Lobby_LeaveLobbyOptions updated, check new fields");
 
@@ -1451,7 +1451,7 @@ TFuture<TDefaultErrorResult<FLobbiesLeaveLobbyImpl>> FLobbiesEOSGS::LeaveLobbyIm
 TFuture<TDefaultErrorResult<FLobbiesDestroyLobbyImpl>> FLobbiesEOSGS::DestroyLobbyImpl(FLobbiesDestroyLobbyImpl::Params&& Params)
 {
 	// Check prerequisites.
-	if (!Services.Get<FAuthEOSGS>()->IsLoggedIn(Params.LocalUserId))
+	if (!Services.Get<FAuthEOSGS>()->IsLoggedIn(Params.LocalAccountId))
 	{
 		return MakeFulfilledPromise<TDefaultErrorResult<FLobbiesDestroyLobbyImpl>>(Errors::InvalidUser()).GetFuture();
 	}
@@ -1463,7 +1463,7 @@ TFuture<TDefaultErrorResult<FLobbiesDestroyLobbyImpl>> FLobbiesEOSGS::DestroyLob
 
 	EOS_Lobby_DestroyLobbyOptions DestroyLobbyOptions = {};
 	DestroyLobbyOptions.ApiVersion = EOS_LOBBY_DESTROYLOBBY_API_LATEST;
-	DestroyLobbyOptions.LocalUserId = GetProductUserIdChecked(Params.LocalUserId);
+	DestroyLobbyOptions.LocalUserId = GetProductUserIdChecked(Params.LocalAccountId);
 	DestroyLobbyOptions.LobbyId = Params.LobbyData->GetLobbyIdEOS();
 	static_assert(EOS_LOBBY_DESTROYLOBBY_API_LATEST == 1, "EOS_Lobby_DestroyLobbyOptions updated, check new fields");
 
@@ -1489,7 +1489,7 @@ TFuture<TDefaultErrorResult<FLobbiesDestroyLobbyImpl>> FLobbiesEOSGS::DestroyLob
 TFuture<TDefaultErrorResult<FLobbiesInviteLobbyMemberImpl>> FLobbiesEOSGS::InviteLobbyMemberImpl(FLobbiesInviteLobbyMemberImpl::Params&& Params)
 {
 	// Check prerequisites.
-	if (!Services.Get<FAuthEOSGS>()->IsLoggedIn(Params.LocalUserId))
+	if (!Services.Get<FAuthEOSGS>()->IsLoggedIn(Params.LocalAccountId))
 	{
 		return MakeFulfilledPromise<TDefaultErrorResult<FLobbiesInviteLobbyMemberImpl>>(Errors::InvalidUser()).GetFuture();
 	}
@@ -1501,8 +1501,8 @@ TFuture<TDefaultErrorResult<FLobbiesInviteLobbyMemberImpl>> FLobbiesEOSGS::Invit
 
 	EOS_Lobby_SendInviteOptions SendInviteOptions = {};
 	SendInviteOptions.ApiVersion = EOS_LOBBY_SENDINVITE_API_LATEST;
-	SendInviteOptions.LocalUserId = GetProductUserIdChecked(Params.LocalUserId);
-	SendInviteOptions.TargetUserId = GetProductUserIdChecked(Params.TargetUserId);
+	SendInviteOptions.LocalUserId = GetProductUserIdChecked(Params.LocalAccountId);
+	SendInviteOptions.TargetUserId = GetProductUserIdChecked(Params.TargetAccountId);
 	SendInviteOptions.LobbyId = Params.LobbyData->GetLobbyIdEOS();
 	static_assert(EOS_LOBBY_SENDINVITE_API_LATEST == 1, "EOS_Lobby_SendInviteOptions updated, check new fields");
 
@@ -1528,13 +1528,13 @@ TFuture<TDefaultErrorResult<FLobbiesInviteLobbyMemberImpl>> FLobbiesEOSGS::Invit
 TFuture<TDefaultErrorResult<FLobbiesDeclineLobbyInvitationImpl>> FLobbiesEOSGS::DeclineLobbyInvitationImpl(FLobbiesDeclineLobbyInvitationImpl::Params&& Params)
 {
 	// Check prerequisites.
-	if (!Services.Get<FAuthEOSGS>()->IsLoggedIn(Params.LocalUserId))
+	if (!Services.Get<FAuthEOSGS>()->IsLoggedIn(Params.LocalAccountId))
 	{
 		return MakeFulfilledPromise<TDefaultErrorResult<FLobbiesDeclineLobbyInvitationImpl>>(Errors::InvalidUser()).GetFuture();
 	}
 
 	// Find the active invitation.
-	TSharedPtr<FLobbyInviteDataEOS> InviteData = GetActiveInvite(Params.LocalUserId, Params.LobbyId);
+	TSharedPtr<FLobbyInviteDataEOS> InviteData = GetActiveInvite(Params.LocalAccountId, Params.LobbyId);
 	if (!InviteData)
 	{
 		// Todo: Errors.
@@ -1543,7 +1543,7 @@ TFuture<TDefaultErrorResult<FLobbiesDeclineLobbyInvitationImpl>> FLobbiesEOSGS::
 
 	EOS_Lobby_RejectInviteOptions RejectInviteOptions = {};
 	RejectInviteOptions.ApiVersion = EOS_LOBBY_REJECTINVITE_API_LATEST;
-	RejectInviteOptions.LocalUserId = GetProductUserIdChecked(Params.LocalUserId);
+	RejectInviteOptions.LocalUserId = GetProductUserIdChecked(Params.LocalAccountId);
 	RejectInviteOptions.InviteId = InviteData->GetInviteIdEOS();
 	static_assert(EOS_LOBBY_REJECTINVITE_API_LATEST == 1, "EOS_Lobby_RejectInviteOptions updated, check new fields");
 
@@ -1572,7 +1572,7 @@ TFuture<TDefaultErrorResult<FLobbiesDeclineLobbyInvitationImpl>> FLobbiesEOSGS::
 TFuture<TDefaultErrorResult<FLobbiesKickLobbyMemberImpl>> FLobbiesEOSGS::KickLobbyMemberImpl(FLobbiesKickLobbyMemberImpl::Params&& Params)
 {
 	// Check prerequisites.
-	if (!Services.Get<FAuthEOSGS>()->IsLoggedIn(Params.LocalUserId))
+	if (!Services.Get<FAuthEOSGS>()->IsLoggedIn(Params.LocalAccountId))
 	{
 		return MakeFulfilledPromise<TDefaultErrorResult<FLobbiesKickLobbyMemberImpl>>(Errors::InvalidUser()).GetFuture();
 	}
@@ -1586,8 +1586,8 @@ TFuture<TDefaultErrorResult<FLobbiesKickLobbyMemberImpl>> FLobbiesEOSGS::KickLob
 
 	EOS_Lobby_KickMemberOptions KickMemberOptions = {};
 	KickMemberOptions.ApiVersion = EOS_LOBBY_KICKMEMBER_API_LATEST;
-	KickMemberOptions.LocalUserId = GetProductUserIdChecked(Params.LocalUserId);
-	KickMemberOptions.TargetUserId = GetProductUserIdChecked(Params.TargetUserId);
+	KickMemberOptions.LocalUserId = GetProductUserIdChecked(Params.LocalAccountId);
+	KickMemberOptions.TargetUserId = GetProductUserIdChecked(Params.TargetAccountId);
 	KickMemberOptions.LobbyId = Params.LobbyData->GetLobbyIdEOS();
 	static_assert(EOS_LOBBY_KICKMEMBER_API_LATEST == 1, "EOS_Lobby_KickMemberOptions updated, check new fields");
 
@@ -1613,7 +1613,7 @@ TFuture<TDefaultErrorResult<FLobbiesKickLobbyMemberImpl>> FLobbiesEOSGS::KickLob
 TFuture<TDefaultErrorResult<FLobbiesPromoteLobbyMemberImpl>> FLobbiesEOSGS::PromoteLobbyMemberImpl(FLobbiesPromoteLobbyMemberImpl::Params&& Params)
 {
 	// Check prerequisites.
-	if (!Services.Get<FAuthEOSGS>()->IsLoggedIn(Params.LocalUserId))
+	if (!Services.Get<FAuthEOSGS>()->IsLoggedIn(Params.LocalAccountId))
 	{
 		return MakeFulfilledPromise<TDefaultErrorResult<FLobbiesPromoteLobbyMemberImpl>>(Errors::InvalidUser()).GetFuture();
 	}
@@ -1627,8 +1627,8 @@ TFuture<TDefaultErrorResult<FLobbiesPromoteLobbyMemberImpl>> FLobbiesEOSGS::Prom
 
 	EOS_Lobby_PromoteMemberOptions PromoteMemberOptions = {};
 	PromoteMemberOptions.ApiVersion = EOS_LOBBY_PROMOTEMEMBER_API_LATEST;
-	PromoteMemberOptions.LocalUserId = GetProductUserIdChecked(Params.LocalUserId);
-	PromoteMemberOptions.TargetUserId = GetProductUserIdChecked(Params.TargetUserId);
+	PromoteMemberOptions.LocalUserId = GetProductUserIdChecked(Params.LocalAccountId);
+	PromoteMemberOptions.TargetUserId = GetProductUserIdChecked(Params.TargetAccountId);
 	PromoteMemberOptions.LobbyId = Params.LobbyData->GetLobbyIdEOS();
 	static_assert(EOS_LOBBY_PROMOTEMEMBER_API_LATEST == 1, "EOS_Lobby_PromoteMemberOptions updated, check new fields");
 
@@ -1654,7 +1654,7 @@ TFuture<TDefaultErrorResult<FLobbiesPromoteLobbyMemberImpl>> FLobbiesEOSGS::Prom
 TFuture<TDefaultErrorResult<FLobbiesModifyLobbyDataImpl>> FLobbiesEOSGS::ModifyLobbyDataImpl(FLobbiesModifyLobbyDataImpl::Params&& Params)
 {
 	// Check prerequisites.
-	if (!Services.Get<FAuthEOSGS>()->IsLoggedIn(Params.LocalUserId))
+	if (!Services.Get<FAuthEOSGS>()->IsLoggedIn(Params.LocalAccountId))
 	{
 		return MakeFulfilledPromise<TDefaultErrorResult<FLobbiesModifyLobbyDataImpl>>(Errors::InvalidUser()).GetFuture();
 	}
@@ -1664,24 +1664,24 @@ TFuture<TDefaultErrorResult<FLobbiesModifyLobbyDataImpl>> FLobbiesEOSGS::ModifyL
 		return MakeFulfilledPromise<TDefaultErrorResult<FLobbiesModifyLobbyDataImpl>>(Errors::InvalidParams()).GetFuture();
 	}
 
-	TSharedPtr<FLobbyDetailsEOS> LobbyDetails = Params.LobbyData->GetUserLobbyDetails(Params.LocalUserId);
+	TSharedPtr<FLobbyDetailsEOS> LobbyDetails = Params.LobbyData->GetUserLobbyDetails(Params.LocalAccountId);
 	if (!LobbyDetails)
 	{
 		return MakeFulfilledPromise<TDefaultErrorResult<FLobbiesModifyLobbyDataImpl>>(Errors::InvalidParams()).GetFuture();
 	}
 
-	if (Params.LocalUserId != Params.LobbyData->GetClientLobbyData()->GetPublicData().OwnerAccountId)
+	if (Params.LocalAccountId != Params.LobbyData->GetClientLobbyData()->GetPublicData().OwnerAccountId)
 	{
 		return MakeFulfilledPromise<TDefaultErrorResult<FLobbiesModifyLobbyDataImpl>>(Errors::InvalidParams()).GetFuture();
 	}
 
 	UE_LOG(LogTemp, Warning, TEXT("[FLobbiesEOSGS::ModifyLobbyDataImpl] Start. Lobby: %s, Member: %s"),
-		*Params.LobbyData->GetLobbyId(), *ToLogString(Params.LocalUserId));
+		*Params.LobbyData->GetLobbyId(), *ToLogString(Params.LocalAccountId));
 
 	TPromise<TDefaultErrorResult<FLobbiesModifyLobbyDataImpl>> Promise;
 	TFuture<TDefaultErrorResult<FLobbiesModifyLobbyDataImpl>> Future = Promise.GetFuture();
 
-	LobbyDetails->ApplyLobbyDataUpdateFromLocalChanges(Params.LocalUserId, *Params.Changes)
+	LobbyDetails->ApplyLobbyDataUpdateFromLocalChanges(Params.LocalAccountId, *Params.Changes)
 	.Then([Promise = MoveTemp(Promise)](TFuture<EOS_EResult>&& Future) mutable
 	{
 		UE_LOG(LogTemp, Warning, TEXT("[FLobbiesEOSGS::ModifyLobbyDataImpl] Complete. Result: %s"), *LexToString(Future.Get()));
@@ -1703,7 +1703,7 @@ TFuture<TDefaultErrorResult<FLobbiesModifyLobbyDataImpl>> FLobbiesEOSGS::ModifyL
 TFuture<TDefaultErrorResult<FLobbiesModifyLobbyMemberDataImpl>> FLobbiesEOSGS::ModifyLobbyMemberDataImpl(FLobbiesModifyLobbyMemberDataImpl::Params&& Params)
 {
 	// Check prerequisites.
-	if (!Services.Get<FAuthEOSGS>()->IsLoggedIn(Params.LocalUserId))
+	if (!Services.Get<FAuthEOSGS>()->IsLoggedIn(Params.LocalAccountId))
 	{
 		return MakeFulfilledPromise<TDefaultErrorResult<FLobbiesModifyLobbyMemberDataImpl>>(Errors::InvalidUser()).GetFuture();
 	}
@@ -1713,7 +1713,7 @@ TFuture<TDefaultErrorResult<FLobbiesModifyLobbyMemberDataImpl>> FLobbiesEOSGS::M
 		return MakeFulfilledPromise<TDefaultErrorResult<FLobbiesModifyLobbyMemberDataImpl>>(Errors::InvalidParams()).GetFuture();
 	}
 
-	TSharedPtr<FLobbyDetailsEOS> LobbyDetails = Params.LobbyData->GetUserLobbyDetails(Params.LocalUserId);
+	TSharedPtr<FLobbyDetailsEOS> LobbyDetails = Params.LobbyData->GetUserLobbyDetails(Params.LocalAccountId);
 	if (!LobbyDetails)
 	{
 		return MakeFulfilledPromise<TDefaultErrorResult<FLobbiesModifyLobbyMemberDataImpl>>(Errors::InvalidParams()).GetFuture();
@@ -1722,7 +1722,7 @@ TFuture<TDefaultErrorResult<FLobbiesModifyLobbyMemberDataImpl>> FLobbiesEOSGS::M
 	TPromise<TDefaultErrorResult<FLobbiesModifyLobbyMemberDataImpl>> Promise;
 	TFuture<TDefaultErrorResult<FLobbiesModifyLobbyMemberDataImpl>> Future = Promise.GetFuture();
 
-	LobbyDetails->ApplyLobbyMemberDataUpdateFromLocalChanges(Params.LocalUserId, *Params.Changes)
+	LobbyDetails->ApplyLobbyMemberDataUpdateFromLocalChanges(Params.LocalAccountId, *Params.Changes)
 	.Then([Promise = MoveTemp(Promise)](TFuture<EOS_EResult>&& Future) mutable
 	{
 		// Todo: Handle "no change" better.
diff --git a/Engine/Plugins/Online/Experimental/OnlineServicesEOSGS/Source/Private/Online/LobbiesEOSGSTypes.cpp b/Engine/Plugins/Online/Experimental/OnlineServicesEOSGS/Source/Private/Online/LobbiesEOSGSTypes.cpp
--- a/Engine/Plugins/Online/Experimental/OnlineServicesEOSGS/Source/Private/Online/LobbiesEOSGSTypes.cpp
+++ b/Engine/Plugins/Online/Experimental/OnlineServicesEOSGS/Source/Private/Online/LobbiesEOSGSTypes.cpp
@@ -147,7 +147,7 @@ const EOS_HLobbyDetails FLobbyDetailsEOS::InvalidLobbyDetailsHandle = {};
 
 TDefaultErrorResultInternal<TSharedRef<FLobbyDetailsEOS>> FLobbyDetailsEOS::CreateFromLobbyId(
 	const TSharedRef<FLobbyPrerequisitesEOS>& Prerequisites,
-	FAccountId LocalUserId,
+	FAccountId LocalAccountId,
 	EOS_LobbyId LobbyId)
 {
 	EOS_HLobbyDetails LobbyDetailsHandle = {};
@@ -155,7 +155,7 @@ TDefaultErrorResultInternal<TSharedRef<FLobbyDetailsEOS>> FLobbyDetailsEOS::Crea
 	EOS_Lobby_CopyLobbyDetailsHandleOptions Options;
 	Options.ApiVersion = EOS_LOBBY_COPYLOBBYDETAILSHANDLE_API_LATEST;
 	Options.LobbyId = LobbyId;
-	Options.LocalUserId = GetProductUserIdChecked(LocalUserId);
+	Options.LocalUserId = GetProductUserIdChecked(LocalAccountId);
 	static_assert(EOS_LOBBY_COPYLOBBYDETAILSHANDLE_API_LATEST == 1, "EOS_Lobby_CopyLobbyDetailsHandleOptions updated, check new fields");
 
 	EOS_EResult EOSResult = EOS_Lobby_CopyLobbyDetailsHandle(Prerequisites->LobbyInterfaceHandle, &Options, &LobbyDetailsHandle);
@@ -171,12 +171,12 @@ TDefaultErrorResultInternal<TSharedRef<FLobbyDetailsEOS>> FLobbyDetailsEOS::Crea
 		return TDefaultErrorResultInternal<TSharedRef<FLobbyDetailsEOS>>(MoveTemp(LobbyDetailsInfoResult.GetErrorValue()));
 	}
 
-	return TDefaultErrorResultInternal<TSharedRef<FLobbyDetailsEOS>>(MakeShared<FLobbyDetailsEOS>(Prerequisites, LobbyDetailsInfoResult.GetOkValue(), LocalUserId, ELobbyDetailsSource::Active, LobbyDetailsHandle));
+	return TDefaultErrorResultInternal<TSharedRef<FLobbyDetailsEOS>>(MakeShared<FLobbyDetailsEOS>(Prerequisites, LobbyDetailsInfoResult.GetOkValue(), LocalAccountId, ELobbyDetailsSource::Active, LobbyDetailsHandle));
 }
 
 TDefaultErrorResultInternal<TSharedRef<FLobbyDetailsEOS>> FLobbyDetailsEOS::CreateFromInviteId(
 	const TSharedRef<FLobbyPrerequisitesEOS>& Prerequisites,
-	FAccountId LocalUserId,
+	FAccountId LocalAccountId,
 	const char* InviteId)
 {
 	EOS_HLobbyDetails LobbyDetailsHandle = {};
@@ -199,12 +199,12 @@ TDefaultErrorResultInternal<TSharedRef<FLobbyDetailsEOS>> FLobbyDetailsEOS::Crea
 		return TDefaultErrorResultInternal<TSharedRef<FLobbyDetailsEOS>>(MoveTemp(LobbyDetailsInfoResult.GetErrorValue()));
 	}
 
-	return TDefaultErrorResultInternal<TSharedRef<FLobbyDetailsEOS>>(MakeShared<FLobbyDetailsEOS>(Prerequisites, LobbyDetailsInfoResult.GetOkValue(), LocalUserId, ELobbyDetailsSource::Invite, LobbyDetailsHandle));
+	return TDefaultErrorResultInternal<TSharedRef<FLobbyDetailsEOS>>(MakeShared<FLobbyDetailsEOS>(Prerequisites, LobbyDetailsInfoResult.GetOkValue(), LocalAccountId, ELobbyDetailsSource::Invite, LobbyDetailsHandle));
 }
 
 TDefaultErrorResultInternal<TSharedRef<FLobbyDetailsEOS>> FLobbyDetailsEOS::CreateFromUiEventId(
 	const TSharedRef<FLobbyPrerequisitesEOS>& Prerequisites,
-	FAccountId LocalUserId,
+	FAccountId LocalAccountId,
 	EOS_UI_EventId UiEventId)
 {
 	EOS_HLobbyDetails LobbyDetailsHandle = {};
@@ -227,12 +227,12 @@ TDefaultErrorResultInternal<TSharedRef<FLobbyDetailsEOS>> FLobbyDetailsEOS::Crea
 		return TDefaultErrorResultInternal<TSharedRef<FLobbyDetailsEOS>>(MoveTemp(LobbyDetailsInfoResult.GetErrorValue()));
 	}
 
-	return TDefaultErrorResultInternal<TSharedRef<FLobbyDetailsEOS>>(MakeShared<FLobbyDetailsEOS>(Prerequisites, LobbyDetailsInfoResult.GetOkValue(), LocalUserId, ELobbyDetailsSource::UiEvent, LobbyDetailsHandle));
+	return TDefaultErrorResultInternal<TSharedRef<FLobbyDetailsEOS>>(MakeShared<FLobbyDetailsEOS>(Prerequisites, LobbyDetailsInfoResult.GetOkValue(), LocalAccountId, ELobbyDetailsSource::UiEvent, LobbyDetailsHandle));
 }
 
 TDefaultErrorResultInternal<TSharedRef<FLobbyDetailsEOS>> FLobbyDetailsEOS::CreateFromSearchResult(
 	const TSharedRef<FLobbyPrerequisitesEOS>& Prerequisites,
-	FAccountId LocalUserId,
+	FAccountId LocalAccountId,
 	EOS_HLobbySearch SearchHandle,
 	uint32_t ResultIndex)
 {
@@ -256,7 +256,7 @@ TDefaultErrorResultInternal<TSharedRef<FLobbyDetailsEOS>> FLobbyDetailsEOS::Crea
 		return TDefaultErrorResultInternal<TSharedRef<FLobbyDetailsEOS>>(MoveTemp(LobbyDetailsInfoResult.GetErrorValue()));
 	}
 
-	return TDefaultErrorResultInternal<TSharedRef<FLobbyDetailsEOS>>(MakeShared<FLobbyDetailsEOS>(Prerequisites, LobbyDetailsInfoResult.GetOkValue(), LocalUserId, ELobbyDetailsSource::Search, LobbyDetailsHandle));
+	return TDefaultErrorResultInternal<TSharedRef<FLobbyDetailsEOS>>(MakeShared<FLobbyDetailsEOS>(Prerequisites, LobbyDetailsInfoResult.GetOkValue(), LocalAccountId, ELobbyDetailsSource::Search, LobbyDetailsHandle));
 }
 
 FLobbyDetailsEOS::~FLobbyDetailsEOS()
@@ -418,7 +418,7 @@ TDefaultErrorResultInternal<TSharedRef<FClientLobbyMemberSnapshot>> FLobbyDetail
 }
 
 TFuture<EOS_EResult> FLobbyDetailsEOS::ApplyLobbyDataUpdateFromLocalChanges(
-	FAccountId LocalUserId,
+	FAccountId LocalAccountId,
 	const FClientLobbyDataChanges& Changes) const
 {
 	EOS_HLobbyModification LobbyModificationHandle = nullptr;
@@ -431,7 +431,7 @@ TFuture<EOS_EResult> FLobbyDetailsEOS::ApplyLobbyDataUpdateFromLocalChanges(
 	// Create lobby modification handle.
 	EOS_Lobby_UpdateLobbyModificationOptions ModificationOptions = {};
 	ModificationOptions.ApiVersion = EOS_LOBBY_UPDATELOBBYMODIFICATION_API_LATEST;
-	ModificationOptions.LocalUserId = GetProductUserIdChecked(LocalUserId);
+	ModificationOptions.LocalUserId = GetProductUserIdChecked(LocalAccountId);
 	ModificationOptions.LobbyId = GetInfo()->GetLobbyId();
 	static_assert(EOS_LOBBY_UPDATELOBBYMODIFICATION_API_LATEST == 1, "EOS_Lobby_UpdateLobbyModificationOptions updated, check new fields");
 
@@ -513,7 +513,7 @@ TFuture<EOS_EResult> FLobbyDetailsEOS::ApplyLobbyDataUpdateFromLocalChanges(
 }
 
 TFuture<EOS_EResult> FLobbyDetailsEOS::ApplyLobbyMemberDataUpdateFromLocalChanges(
-	FAccountId LocalUserId,
+	FAccountId LocalAccountId,
 	const FClientLobbyMemberDataChanges& Changes) const
 {
 	EOS_HLobbyModification LobbyModificationHandle = {};
@@ -526,7 +526,7 @@ TFuture<EOS_EResult> FLobbyDetailsEOS::ApplyLobbyMemberDataUpdateFromLocalChange
 	// Create lobby modification handle.
 	EOS_Lobby_UpdateLobbyModificationOptions ModificationOptions = {};
 	ModificationOptions.ApiVersion = EOS_LOBBY_UPDATELOBBYMODIFICATION_API_LATEST;
-	ModificationOptions.LocalUserId = GetProductUserIdChecked(LocalUserId);
+	ModificationOptions.LocalUserId = GetProductUserIdChecked(LocalAccountId);
 	ModificationOptions.LobbyId = GetInfo()->GetLobbyId();
 	static_assert(EOS_LOBBY_UPDATELOBBYMODIFICATION_API_LATEST == 1, "EOS_Lobby_UpdateLobbyModificationOptions updated, check new fields");
 
@@ -595,12 +595,12 @@ TFuture<EOS_EResult> FLobbyDetailsEOS::ApplyLobbyMemberDataUpdateFromLocalChange
 FLobbyDetailsEOS::FLobbyDetailsEOS(
 	const TSharedRef<FLobbyPrerequisitesEOS>& Prerequisites,
 	const TSharedRef<FLobbyDetailsInfoEOS>& LobbyDetailsInfo,
-	FAccountId LocalUserId,
+	FAccountId LocalAccountId,
 	ELobbyDetailsSource LobbyDetailsSource,
 	EOS_HLobbyDetails LobbyDetailsHandle)
 	: Prerequisites(Prerequisites)
 	, LobbyDetailsInfo(LobbyDetailsInfo)
-	, AssociatedLocalUser(LocalUserId)
+	, AssociatedLocalUser(LocalAccountId)
 	, LobbyDetailsSource(LobbyDetailsSource)
 	, LobbyDetailsHandle(LobbyDetailsHandle)
 {
@@ -642,9 +642,9 @@ FLobbyDataEOS::~FLobbyDataEOS()
 	}
 }
 
-void FLobbyDataEOS::AddUserLobbyDetails(FAccountId LocalUserId, const TSharedPtr<FLobbyDetailsEOS>& LobbyDetails)
+void FLobbyDataEOS::AddUserLobbyDetails(FAccountId LocalAccountId, const TSharedPtr<FLobbyDetailsEOS>& LobbyDetails)
 {
-	if (TSharedPtr<FLobbyDetailsEOS> ExistingDetails = GetUserLobbyDetails(LocalUserId))
+	if (TSharedPtr<FLobbyDetailsEOS> ExistingDetails = GetUserLobbyDetails(LocalAccountId))
 	{
 		if (ExistingDetails->GetDetailsSource() < LobbyDetails->GetDetailsSource())
 		{
@@ -652,12 +652,12 @@ void FLobbyDataEOS::AddUserLobbyDetails(FAccountId LocalUserId, const TSharedPtr
 		}
 	}
 
-	UserLobbyDetails.Add(LocalUserId, LobbyDetails);
+	UserLobbyDetails.Add(LocalAccountId, LobbyDetails);
 }
 
-TSharedPtr<FLobbyDetailsEOS> FLobbyDataEOS::GetUserLobbyDetails(FAccountId LocalUserId) const
+TSharedPtr<FLobbyDetailsEOS> FLobbyDataEOS::GetUserLobbyDetails(FAccountId LocalAccountId) const
 {
-	const TSharedPtr<FLobbyDetailsEOS>* Result = UserLobbyDetails.Find(LocalUserId);
+	const TSharedPtr<FLobbyDetailsEOS>* Result = UserLobbyDetails.Find(LocalAccountId);
 	return Result ? *Result : TSharedPtr<FLobbyDetailsEOS>();
 }
 
@@ -759,11 +759,11 @@ TSharedPtr<FLobbyDataEOS> FLobbyDataRegistryEOS::Find(FOnlineLobbyIdHandle Lobby
 	return Result ? Result->Pin() : TSharedPtr<FLobbyDataEOS>();
 }
 
-TFuture<TDefaultErrorResultInternal<TSharedRef<FLobbyDataEOS>>> FLobbyDataRegistryEOS::FindOrCreateFromLobbyDetails(FAccountId LocalUserId, const TSharedRef<FLobbyDetailsEOS>& LobbyDetails)
+TFuture<TDefaultErrorResultInternal<TSharedRef<FLobbyDataEOS>>> FLobbyDataRegistryEOS::FindOrCreateFromLobbyDetails(FAccountId LocalAccountId, const TSharedRef<FLobbyDetailsEOS>& LobbyDetails)
 {
 	if (TSharedPtr<FLobbyDataEOS> FindResult = Find(LobbyDetails->GetInfo()->GetLobbyId()))
 	{
-		FindResult->AddUserLobbyDetails(LocalUserId, LobbyDetails);
+		FindResult->AddUserLobbyDetails(LocalAccountId, LobbyDetails);
 		return MakeFulfilledPromise<TDefaultErrorResultInternal<TSharedRef<FLobbyDataEOS>>>(FindResult.ToSharedRef()).GetFuture();
 	}
 
@@ -772,7 +772,7 @@ TFuture<TDefaultErrorResultInternal<TSharedRef<FLobbyDataEOS>>> FLobbyDataRegist
 
 	const FOnlineLobbyIdHandle LobbyId = FOnlineLobbyIdHandle(EOnlineServices::Epic, NextHandleIndex++);
 	FLobbyDataEOS::Create(LobbyId, LobbyDetails, MakeUnregisterFn())
-	.Then([WeakThis = AsWeak(), Promise = MoveTemp(Promise), LocalUserId, LobbyDetails](TFuture<TDefaultErrorResultInternal<TSharedRef<FLobbyDataEOS>>>&& Future) mutable
+	.Then([WeakThis = AsWeak(), Promise = MoveTemp(Promise), LocalAccountId, LobbyDetails](TFuture<TDefaultErrorResultInternal<TSharedRef<FLobbyDataEOS>>>&& Future) mutable
 	{
 		if (TSharedPtr<FLobbyDataRegistryEOS> StrongThis = WeakThis.Pin())
 		{
@@ -782,7 +782,7 @@ TFuture<TDefaultErrorResultInternal<TSharedRef<FLobbyDataEOS>>> FLobbyDataRegist
 			}
 		}
 
-		Future.Get().GetOkValue()->AddUserLobbyDetails(LocalUserId, LobbyDetails);
+		Future.Get().GetOkValue()->AddUserLobbyDetails(LocalAccountId, LobbyDetails);
 		Promise.EmplaceValue(MoveTempIfPossible(Future.Get()));
 	});
 
@@ -818,11 +818,11 @@ FLobbyDataEOS::FUnregisterFn FLobbyDataRegistryEOS::MakeUnregisterFn()
 TFuture<TDefaultErrorResultInternal<TSharedRef<FLobbyInviteDataEOS>>> FLobbyInviteDataEOS::CreateFromInviteId(
 	const TSharedRef<FLobbyPrerequisitesEOS>& Prerequisites,
 	const TSharedRef<FLobbyDataRegistryEOS>& LobbyDataRegistry,
-	FAccountId LocalUserId,
+	FAccountId LocalAccountId,
 	const char* InviteIdEOS,
 	EOS_ProductUserId Sender)
 {
-	TDefaultErrorResultInternal<TSharedRef<FLobbyDetailsEOS>> LobbyDetailsResult = FLobbyDetailsEOS::CreateFromInviteId(Prerequisites, LocalUserId, InviteIdEOS);
+	TDefaultErrorResultInternal<TSharedRef<FLobbyDetailsEOS>> LobbyDetailsResult = FLobbyDetailsEOS::CreateFromInviteId(Prerequisites, LocalAccountId, InviteIdEOS);
 	if (LobbyDetailsResult.IsError())
 	{
 		return MakeFulfilledPromise<TDefaultErrorResultInternal<TSharedRef<FLobbyInviteDataEOS>>>(MoveTemp(LobbyDetailsResult.GetErrorValue())).GetFuture();
@@ -834,8 +834,8 @@ TFuture<TDefaultErrorResultInternal<TSharedRef<FLobbyInviteDataEOS>>> FLobbyInvi
 
 	// Search for existing lobby data so that the LobbyIdHandle will match.
 	TSharedRef<FLobbyInviteIdEOS> InviteId = MakeShared<FLobbyInviteIdEOS>(InviteIdEOS);
-	LobbyDataRegistry->FindOrCreateFromLobbyDetails(LocalUserId, LobbyDetails)
-	.Then([Promise = MoveTemp(Promise), InviteId, LocalUserId, Sender, LobbyDetails](TFuture<TDefaultErrorResultInternal<TSharedRef<FLobbyDataEOS>>>&& Future) mutable
+	LobbyDataRegistry->FindOrCreateFromLobbyDetails(LocalAccountId, LobbyDetails)
+	.Then([Promise = MoveTemp(Promise), InviteId, LocalAccountId, Sender, LobbyDetails](TFuture<TDefaultErrorResultInternal<TSharedRef<FLobbyDataEOS>>>&& Future) mutable
 	{
 		if (Future.Get().IsError())
 		{
@@ -844,15 +844,15 @@ TFuture<TDefaultErrorResultInternal<TSharedRef<FLobbyInviteDataEOS>>> FLobbyInvi
 		}
 
 		// Once the lobby data has been resolved the FAccountId for the sender is expected to be in the AccountID cache.
-		const FAccountId SenderUserId = FindAccountId(Sender);
-		if (!SenderUserId.IsValid())
+		const FAccountId SenderAccountId = FindAccountId(Sender);
+		if (!SenderAccountId.IsValid())
 		{
 			// Todo: Errors.
 			Promise.EmplaceValue(MoveTemp(Future.Get().GetErrorValue()));
 			return;
 		}
 
-		Promise.EmplaceValue(MakeShared<FLobbyInviteDataEOS>(InviteId, LocalUserId, SenderUserId, LobbyDetails, Future.Get().GetOkValue()));
+		Promise.EmplaceValue(MakeShared<FLobbyInviteDataEOS>(InviteId, LocalAccountId, SenderAccountId, LobbyDetails, Future.Get().GetOkValue()));
 	});
 
 	return Future;
@@ -978,11 +978,11 @@ TFuture<TDefaultErrorResultInternal<TSharedRef<FLobbySearchEOS>>> FLobbySearchEO
 
 	EOS_LobbySearch_FindOptions FindOptions = {};
 	FindOptions.ApiVersion = EOS_LOBBYSEARCH_FIND_API_LATEST;
-	FindOptions.LocalUserId = GetProductUserIdChecked(Params.LocalUserId);
+	FindOptions.LocalUserId = GetProductUserIdChecked(Params.LocalAccountId);
 	static_assert(EOS_LOBBYSEARCH_FIND_API_LATEST == 1, "EOS_LobbySearch_FindOptions updated, check new fields");
 
 	EOS_Async(EOS_LobbySearch_Find, SearchHandle->Get(), FindOptions,
-	[Promise = MoveTemp(Promise), Prerequisites, LobbyRegistry, LocalUserId = Params.LocalUserId, SearchHandle]
+	[Promise = MoveTemp(Promise), Prerequisites, LobbyRegistry, LocalAccountId = Params.LocalAccountId, SearchHandle]
 	(const EOS_LobbySearch_FindCallbackInfo* CallbackInfo) mutable
 	{
 		if (CallbackInfo->ResultCode != EOS_EResult::EOS_Success)
@@ -1003,7 +1003,7 @@ TFuture<TDefaultErrorResultInternal<TSharedRef<FLobbySearchEOS>>> FLobbySearchEO
 		for (uint32_t SearchResultIndex = 0; SearchResultIndex < NumSearchResults; ++SearchResultIndex)
 		{
 			TDefaultErrorResultInternal<TSharedRef<FLobbyDetailsEOS>> Result = FLobbyDetailsEOS::CreateFromSearchResult(
-				Prerequisites, LocalUserId, SearchHandle->Get(), SearchResultIndex);
+				Prerequisites, LocalAccountId, SearchHandle->Get(), SearchResultIndex);
 			if (Result.IsError())
 			{
 				// todo: errors
@@ -1014,7 +1014,7 @@ TFuture<TDefaultErrorResultInternal<TSharedRef<FLobbySearchEOS>>> FLobbySearchEO
 			TPromise<TDefaultErrorResultInternal<TSharedRef<FLobbyDataEOS>>> ResolveLobbyDetailsPromise;
 			ResolvedLobbyDetails.Add(ResolveLobbyDetailsPromise.GetFuture());
 
-			LobbyRegistry->FindOrCreateFromLobbyDetails(LocalUserId, Result.GetOkValue())
+			LobbyRegistry->FindOrCreateFromLobbyDetails(LocalAccountId, Result.GetOkValue())
 			.Then([ResolveLobbyDetailsPromise = MoveTemp(ResolveLobbyDetailsPromise)](TFuture<TDefaultErrorResultInternal<TSharedRef<FLobbyDataEOS>>>&& Future) mutable
 			{
 				ResolveLobbyDetailsPromise.EmplaceValue(MoveTempIfPossible(Future.Get()));
diff --git a/Engine/Plugins/Online/Experimental/OnlineServicesEOSGS/Source/Private/Online/LobbiesEOSGSTypes.h b/Engine/Plugins/Online/Experimental/OnlineServicesEOSGS/Source/Private/Online/LobbiesEOSGSTypes.h
--- a/Engine/Plugins/Online/Experimental/OnlineServicesEOSGS/Source/Private/Online/LobbiesEOSGSTypes.h
+++ b/Engine/Plugins/Online/Experimental/OnlineServicesEOSGS/Source/Private/Online/LobbiesEOSGSTypes.h
@@ -213,19 +213,19 @@ public:
 
 	static TDefaultErrorResultInternal<TSharedRef<FLobbyDetailsEOS>> CreateFromLobbyId(
 		const TSharedRef<FLobbyPrerequisitesEOS>& Prerequisites,
-		FAccountId LocalUserId,
+		FAccountId LocalAccountId,
 		EOS_LobbyId LobbyId);
 	static TDefaultErrorResultInternal<TSharedRef<FLobbyDetailsEOS>> CreateFromInviteId(
 		const TSharedRef<FLobbyPrerequisitesEOS>& Prerequisites,
-		FAccountId LocalUserId,
+		FAccountId LocalAccountId,
 		const char* InviteId);
 	static TDefaultErrorResultInternal<TSharedRef<FLobbyDetailsEOS>> CreateFromUiEventId(
 		const TSharedRef<FLobbyPrerequisitesEOS>& Prerequisites,
-		FAccountId LocalUserId,
+		FAccountId LocalAccountId,
 		EOS_UI_EventId UiEventId);
 	static TDefaultErrorResultInternal<TSharedRef<FLobbyDetailsEOS>> CreateFromSearchResult(
 		const TSharedRef<FLobbyPrerequisitesEOS>& Prerequisites,
-		FAccountId LocalUserId,
+		FAccountId LocalAccountId,
 		EOS_HLobbySearch SearchHandle,
 		uint32_t ResultIndex);
 
@@ -247,19 +247,19 @@ public:
 	 * Retrieve lobby member data snapshot from the EOS lobby details object.
 	 * The lobby schema will be used to translate attribute data before returning.
 	 */
-	TDefaultErrorResultInternal<TSharedRef<FClientLobbyMemberSnapshot>> GetLobbyMemberSnapshot(FAccountId MemberHandle) const;
+	TDefaultErrorResultInternal<TSharedRef<FClientLobbyMemberSnapshot>> GetLobbyMemberSnapshot(FAccountId MemberAccountId) const;
 
 	/**
 	 * Apply client side lobby changes to the lobby service.
 	 * The lobby schema will be used to translate any changed attributes before sending.
 	 */
-	TFuture<EOS_EResult> ApplyLobbyDataUpdateFromLocalChanges(FAccountId LocalUserId, const FClientLobbyDataChanges& Changes) const;
+	TFuture<EOS_EResult> ApplyLobbyDataUpdateFromLocalChanges(FAccountId LocalAccountId, const FClientLobbyDataChanges& Changes) const;
 
 	/**
 	 * Apply client side lobby member changes to the lobby service.
 	 * The lobby schema will be used to translate any changed attributes before sending.
 	 */
-	TFuture<EOS_EResult> ApplyLobbyMemberDataUpdateFromLocalChanges(FAccountId LocalUserId, const FClientLobbyMemberDataChanges& Changes) const;
+	TFuture<EOS_EResult> ApplyLobbyMemberDataUpdateFromLocalChanges(FAccountId LocalAccountId, const FClientLobbyMemberDataChanges& Changes) const;
 
 private:
 	template <typename, ESPMode>
@@ -267,7 +267,7 @@ private:
 	FLobbyDetailsEOS(
 		const TSharedRef<FLobbyPrerequisitesEOS>& Prerequisites,
 		const TSharedRef<FLobbyDetailsInfoEOS>& LobbyDetailsInfo,
-		FAccountId LocalUserId,
+		FAccountId LocalAccountId,
 		ELobbyDetailsSource LobbyDetailsSource,
 		EOS_HLobbyDetails LobbyDetailsHandle);
 
@@ -334,8 +334,8 @@ public:
 	EOS_LobbyId GetLobbyIdEOS() const { return LobbyDetailsInfo->GetLobbyId(); }
 	const FString& GetLobbyId() const { return LobbyId; }
 
-	void AddUserLobbyDetails(FAccountId LocalUserId, const TSharedPtr<FLobbyDetailsEOS>& LobbyDetails);
-	TSharedPtr<FLobbyDetailsEOS> GetUserLobbyDetails(FAccountId LocalUserId) const;
+	void AddUserLobbyDetails(FAccountId LocalAccountId, const TSharedPtr<FLobbyDetailsEOS>& LobbyDetails);
+	TSharedPtr<FLobbyDetailsEOS> GetUserLobbyDetails(FAccountId LocalAccountId) const;
 
 	/**
 	 * Active lobby details are needed to process lobby notifications. Search for and return active
@@ -367,7 +367,7 @@ public:
 
 	TSharedPtr<FLobbyDataEOS> Find(EOS_LobbyId EOSLobbyId) const;
 	TSharedPtr<FLobbyDataEOS> Find(FOnlineLobbyIdHandle LobbyIdHandle) const;
-	TFuture<TDefaultErrorResultInternal<TSharedRef<FLobbyDataEOS>>> FindOrCreateFromLobbyDetails(FAccountId LocalUserId, const TSharedRef<FLobbyDetailsEOS>& LobbyDetails);
+	TFuture<TDefaultErrorResultInternal<TSharedRef<FLobbyDataEOS>>> FindOrCreateFromLobbyDetails(FAccountId LocalAccountId, const TSharedRef<FLobbyDetailsEOS>& LobbyDetails);
 
 private:
 	void Register(const TSharedRef<FLobbyDataEOS>& LobbyIdHandleData);
@@ -407,7 +407,7 @@ public:
 	static TFuture<TDefaultErrorResultInternal<TSharedRef<FLobbyInviteDataEOS>>> CreateFromInviteId(
 		const TSharedRef<FLobbyPrerequisitesEOS>& Prerequisites,
 		const TSharedRef<FLobbyDataRegistryEOS>& LobbyDataRegistry,
-		FAccountId LocalUserId,
+		FAccountId LocalAccountId,
 		const char* InviteIdEOS,
 		EOS_ProductUserId Sender);
 
diff --git a/Engine/Plugins/Online/Experimental/OnlineServicesEOSGS/Source/Private/Online/OnlineIdEOSGS.cpp b/Engine/Plugins/Online/Experimental/OnlineServicesEOSGS/Source/Private/Online/OnlineIdEOSGS.cpp
--- a/Engine/Plugins/Online/Experimental/OnlineServicesEOSGS/Source/Private/Online/OnlineIdEOSGS.cpp
+++ b/Engine/Plugins/Online/Experimental/OnlineServicesEOSGS/Source/Private/Online/OnlineIdEOSGS.cpp
@@ -23,27 +23,27 @@ FAccountId FOnlineAccountIdRegistryEOSGS::FindAccountId(const EOS_ProductUserId
 	return Registry.FindHandle(ProductUserId);
 }
 
-EOS_ProductUserId FOnlineAccountIdRegistryEOSGS::GetProductUserId(const FAccountId& Handle) const
+EOS_ProductUserId FOnlineAccountIdRegistryEOSGS::GetProductUserId(const FAccountId& AccountId) const
 {
-	return Registry.FindIdValue(Handle);
+	return Registry.FindIdValue(AccountId);
 }
 
-FString FOnlineAccountIdRegistryEOSGS::ToLogString(const FAccountId& Handle) const
+FString FOnlineAccountIdRegistryEOSGS::ToLogString(const FAccountId& AccountId) const
 {
-	if (Registry.ValidateOnlineId(Handle))
+	if (Registry.ValidateOnlineId(AccountId))
 	{
-		EOS_ProductUserId ProductUserId = Registry.FindIdValue(Handle);
+		EOS_ProductUserId ProductUserId = Registry.FindIdValue(AccountId);
 		return LexToString(ProductUserId);
 	}
 	return FString();
 }
 
-TArray<uint8> FOnlineAccountIdRegistryEOSGS::ToReplicationData(const FAccountId& Handle) const
+TArray<uint8> FOnlineAccountIdRegistryEOSGS::ToReplicationData(const FAccountId& AccountId) const
 {
 	TArray<uint8> ReplicationData;
-	if (Registry.ValidateOnlineId(Handle))
+	if (Registry.ValidateOnlineId(AccountId))
 	{
-		EOS_ProductUserId ProductUserId = Registry.FindIdValue(Handle);
+		EOS_ProductUserId ProductUserId = Registry.FindIdValue(AccountId);
 		if (ensure(EOS_ProductUserId_IsValid(ProductUserId)))
 		{
 			char EosBuffer[EOS_PRODUCTUSERID_MAX_LENGTH + 1] = {};
@@ -93,14 +93,14 @@ FAccountId FOnlineAccountIdRegistryEOSGS::FindOrAddAccountId(const EOS_ProductUs
 	return Registry.GetInvalidHandle();
 }
 
-EOS_ProductUserId GetProductUserId(const FAccountId& Handle)
+EOS_ProductUserId GetProductUserId(const FAccountId& AccountId)
 {
-	return FOnlineAccountIdRegistryEOSGS::GetRegistered().GetProductUserId(Handle);
+	return FOnlineAccountIdRegistryEOSGS::GetRegistered().GetProductUserId(AccountId);
 }
 
-EOS_ProductUserId GetProductUserIdChecked(const FAccountId& Handle)
+EOS_ProductUserId GetProductUserIdChecked(const FAccountId& AccountId)
 {
-	EOS_ProductUserId ProductUserId = GetProductUserId(Handle);
+	EOS_ProductUserId ProductUserId = GetProductUserId(AccountId);
 	check(EOS_ProductUserId_IsValid(ProductUserId) == EOS_TRUE);
 	return ProductUserId;
 }
diff --git a/Engine/Plugins/Online/Experimental/OnlineServicesEOSGS/Source/Private/Online/OnlineServicesEOSGS.cpp b/Engine/Plugins/Online/Experimental/OnlineServicesEOSGS/Source/Private/Online/OnlineServicesEOSGS.cpp
--- a/Engine/Plugins/Online/Experimental/OnlineServicesEOSGS/Source/Private/Online/OnlineServicesEOSGS.cpp
+++ b/Engine/Plugins/Online/Experimental/OnlineServicesEOSGS/Source/Private/Online/OnlineServicesEOSGS.cpp
@@ -99,7 +99,7 @@ TOnlineResult<FGetResolvedConnectString> FOnlineServicesEOSGS::GetResolvedConnec
 		ILobbiesPtr LobbiesEOS = GetLobbiesInterface();
 		check(LobbiesEOS);
 
-		TOnlineResult<FGetJoinedLobbies> JoinedLobbies = LobbiesEOS->GetJoinedLobbies({ Params.LocalUserId });
+		TOnlineResult<FGetJoinedLobbies> JoinedLobbies = LobbiesEOS->GetJoinedLobbies({ Params.LocalAccountId });
 		if (JoinedLobbies.IsOk())
 		{
 			for (TSharedRef<const FLobby>& Lobby : JoinedLobbies.GetOkValue().Lobbies)
@@ -130,7 +130,7 @@ TOnlineResult<FGetResolvedConnectString> FOnlineServicesEOSGS::GetResolvedConnec
 		ISessionsPtr SessionsEOS = GetSessionsInterface();
 		check(SessionsEOS);
 
-		TOnlineResult<FGetSessionById> Result = SessionsEOS->GetSessionById({ Params.LocalUserId, Params.SessionId });
+		TOnlineResult<FGetSessionById> Result = SessionsEOS->GetSessionById({ Params.LocalAccountId, Params.SessionId });
 		if (Result.IsOk())
 		{
 #if WITH_ENGINE
@@ -138,7 +138,7 @@ TOnlineResult<FGetResolvedConnectString> FOnlineServicesEOSGS::GetResolvedConnec
 
 			//It should look like this: "EOS:0002aeeb5b2d4388a3752dd6d31222ec:GameNetDriver:97"
 			FString NetDriverName = GetDefault<UNetDriverEOSBase>()->NetDriverName.ToString();
-			FInternetAddrEOS TempAddr(GetProductUserIdChecked(Session->OwnerUserId), NetDriverName, GetTypeHash(NetDriverName));
+			FInternetAddrEOS TempAddr(GetProductUserIdChecked(Session->OwnerAccountId), NetDriverName, GetTypeHash(NetDriverName));
 			return TOnlineResult<FGetResolvedConnectString>({ TempAddr.ToString(true) });
 #else
 			return TOnlineResult<FGetResolvedConnectString>(Errors::NotImplemented());
