OSSv2: Rename Handle, UserId to AccountId when use FAccountId.

diff --git a/Engine/Plugins/Online/Experimental/OnlineServicesEOSGS/Source/Private/Online/SessionsEOSGS.cpp b/Engine/Plugins/Online/Experimental/OnlineServicesEOSGS/Source/Private/Online/SessionsEOSGS.cpp
--- a/Engine/Plugins/Online/Experimental/OnlineServicesEOSGS/Source/Private/Online/SessionsEOSGS.cpp
+++ b/Engine/Plugins/Online/Experimental/OnlineServicesEOSGS/Source/Private/Online/SessionsEOSGS.cpp
@@ -96,12 +96,12 @@ FSessionEOSGS::FSessionEOSGS(const EOS_HSessionDetails& InSessionDetailsHandle)
 						// We retrieve the registered player id
 						const FString& PlayerIdStr = KeyComponents[1];
 						const EOS_ProductUserId ProductUserId = EOS_ProductUserId_FromString(TCHAR_TO_UTF8(*PlayerIdStr));
-						const FAccountId& IdHandle = FindAccountId(ProductUserId);
+						const FAccountId& AccountId = FindAccountId(ProductUserId);
 
 						// This constructor should only be used by BuildSessionFromDetailsHandle, after all user ids in the session have been resolved
-						check(IdHandle.IsValid());
+						check(AccountId.IsValid());
 
-						FRegisteredPlayer& RegisteredPlayer = SessionSettings.RegisteredPlayers.FindOrAdd(IdHandle);
+						FRegisteredPlayer& RegisteredPlayer = SessionSettings.RegisteredPlayers.FindOrAdd(AccountId);
 						
 						FSessionAttributeConverter<ESessionAttributeConversionType::FromService> CustomSettingConverter(*Attribute->Data);
 						const TPair<FSchemaAttributeId, FSchemaVariant>& CustomSettingData = CustomSettingConverter.GetAttributeData();
@@ -124,9 +124,9 @@ FSessionEOSGS::FSessionEOSGS(const EOS_HSessionDetails& InSessionDetailsHandle)
 						// We retrieve the member id
 						FString PlayerIdStr = KeyComponents[0];
 						EOS_ProductUserId ProductUserId = EOS_ProductUserId_FromString(TCHAR_TO_UTF8(*PlayerIdStr));
-						FAccountId IdHandle = FindAccountId(ProductUserId);
+						FAccountId AccountId = FindAccountId(ProductUserId);
 
-						FSessionMember& SessionMember = SessionSettings.SessionMembers.FindOrAdd(IdHandle);
+						FSessionMember& SessionMember = SessionSettings.SessionMembers.FindOrAdd(AccountId);
 
 						// And add the corresponding custom setting
 						FSchemaAttributeId AttributeId(KeyComponents[1]);
@@ -273,12 +273,12 @@ void FSessionsEOSGS::UnregisterEventHandlers()
 
 void FSessionsEOSGS::HandleSessionInviteReceived(const EOS_Sessions_SessionInviteReceivedCallbackInfo* Data)
 {
-	FAccountId LocalUserId = FindAccountId(Data->LocalUserId);
-	if (LocalUserId.IsValid())
+	FAccountId LocalAccountId = FindAccountId(Data->LocalUserId);
+	if (LocalAccountId.IsValid())
 	{
 		FString InviteId(UTF8_TO_TCHAR(Data->InviteId));
 
-		BuildSessionFromInvite(LocalUserId, InviteId)
+		BuildSessionFromInvite(LocalAccountId, InviteId)
 			.OnComplete([this, WeakThis = AsWeak(), SenderPUID = Data->TargetUserId, InviteId](const TOnlineResult<FBuildSessionFromDetailsHandle>& Result)
 		{
 			if (Result.IsOk())
@@ -286,7 +286,7 @@ void FSessionsEOSGS::HandleSessionInviteReceived(const EOS_Sessions_SessionInvit
 				if (const TSharedPtr<ISessions> StrongThis = WeakThis.Pin())
 				{
 					// First and second place in the array will be occupied by the receiver and the sender, respectively, since the same order is kept in the array of resolved ids
-					const FAccountId& ReceiverId = Result.GetOkValue().LocalUserId;
+					const FAccountId& ReceiverId = Result.GetOkValue().LocalAccountId;
 					const FAccountId& SenderId = FindAccountId(SenderPUID);
 
 					FOnlineSessionInviteIdHandle InviteIdHandle = CreateSessionInviteId(InviteId);
@@ -298,7 +298,7 @@ void FSessionsEOSGS::HandleSessionInviteReceived(const EOS_Sessions_SessionInvit
 						SenderId,
 						InviteIdHandle,
 						Session->SessionId
-					});
+						});
 
 					SessionInvitesUserMap.FindOrAdd(ReceiverId).Emplace(InviteIdHandle, SessionInviteRef);
 
@@ -317,12 +317,12 @@ void FSessionsEOSGS::HandleSessionInviteReceived(const EOS_Sessions_SessionInvit
 
 void FSessionsEOSGS::HandleSessionInviteAccepted(const EOS_Sessions_SessionInviteAcceptedCallbackInfo* Data)
 {
-	FAccountId LocalUserId = FindAccountId(Data->LocalUserId);
-	if (LocalUserId.IsValid())
+	FAccountId LocalAccountId = FindAccountId(Data->LocalUserId);
+	if (LocalAccountId.IsValid())
 	{
 		FString InviteId(UTF8_TO_TCHAR(Data->InviteId));
 
-		BuildSessionFromInvite(LocalUserId, InviteId)
+		BuildSessionFromInvite(LocalAccountId, InviteId)
 			.OnComplete([this, WeakThis = AsWeak()](const TOnlineResult<FBuildSessionFromDetailsHandle>& Result)
 		{
 			if (Result.IsOk())
@@ -333,8 +333,8 @@ void FSessionsEOSGS::HandleSessionInviteAccepted(const EOS_Sessions_SessionInvit
 					AllSessionsById.Emplace(Result.GetOkValue().Session->SessionId, Result.GetOkValue().Session);
 
 					FUISessionJoinRequested Event{
-						Result.GetOkValue().LocalUserId,
-						TResult<FOnlineSessionIdHandle, FOnlineError>(Result.GetOkValue().Session->SessionId),
+									Result.GetOkValue().LocalAccountId,
+									TResult<FOnlineSessionIdHandle, FOnlineError>(Result.GetOkValue().Session->SessionId),
 						EUISessionJoinRequestedSource::FromInvitation
 					};
 
@@ -349,10 +349,10 @@ void FSessionsEOSGS::HandleSessionInviteAccepted(const EOS_Sessions_SessionInvit
 
 void FSessionsEOSGS::HandleJoinSessionAccepted(const EOS_Sessions_JoinSessionAcceptedCallbackInfo* Data)
 {
-	FAccountId LocalUserId = FindAccountId(Data->LocalUserId);
-	if (LocalUserId.IsValid())
+	FAccountId LocalAccountId = FindAccountId(Data->LocalUserId);
+	if (LocalAccountId.IsValid())
 	{
-		BuildSessionFromUIEvent(LocalUserId, Data->UiEventId)
+		BuildSessionFromUIEvent(LocalAccountId, Data->UiEventId)
 			.OnComplete([this, WeakThis = AsWeak()](const TOnlineResult<FBuildSessionFromDetailsHandle>& Result)
 		{
 			if (Result.IsOk())
@@ -363,9 +363,9 @@ void FSessionsEOSGS::HandleJoinSessionAccepted(const EOS_Sessions_JoinSessionAcc
 					AllSessionsById.Emplace(Result.GetOkValue().Session->SessionId, Result.GetOkValue().Session);
 
 					FUISessionJoinRequested Event{
-						Result.GetOkValue().LocalUserId,
+						Result.GetOkValue().LocalAccountId,
 						TResult<FOnlineSessionIdHandle, FOnlineError>(Result.GetOkValue().Session->SessionId),
-						EUISessionJoinRequestedSource::Unspecified
+			EUISessionJoinRequestedSource::Unspecified
 					};
 
 					SessionEvents.OnUISessionJoinRequested.Broadcast(Event);
@@ -450,7 +450,7 @@ TOnlineAsyncOpHandle<FCreateSession> FSessionsEOSGS::CreateSession(FCreateSessio
 			CreateSessionModificationOptions.BucketId = BucketIdUtf8.Get();
 		}
 	
-		CreateSessionModificationOptions.LocalUserId = GetProductUserIdChecked(OpParams.LocalUserId);
+		CreateSessionModificationOptions.LocalUserId = GetProductUserIdChecked(OpParams.LocalAccountId);
 		CreateSessionModificationOptions.MaxPlayers = OpParams.SessionSettings.NumMaxPrivateConnections + OpParams.SessionSettings.NumMaxPublicConnections;
 
 		const FTCHARToUTF8 SessionIdUtf8(*OpParams.SessionSettings.SessionIdOverride);
@@ -491,7 +491,7 @@ TOnlineAsyncOpHandle<FCreateSession> FSessionsEOSGS::CreateSession(FCreateSessio
 		{
 			// If the session is created successfully, we'll create the local session object
 			TSharedRef<FSessionEOSGS> NewSessionEOSGSRef = MakeShared<FSessionEOSGS>();
-			NewSessionEOSGSRef->OwnerUserId = OpParams.LocalUserId;
+			NewSessionEOSGSRef->OwnerAccountId = OpParams.LocalAccountId;
 			NewSessionEOSGSRef->SessionSettings = OpParams.SessionSettings;
 			NewSessionEOSGSRef->SessionId = CreateSessionId(Result.GetOkValue().NewSessionId);
 
@@ -499,11 +499,11 @@ TOnlineAsyncOpHandle<FCreateSession> FSessionsEOSGS::CreateSession(FCreateSessio
 
 			LocalSessionsByName.Emplace(OpParams.SessionName, NewSessionEOSGSRef->SessionId);
 
-			NamedSessionUserMap.FindOrAdd(OpParams.LocalUserId).AddUnique(OpParams.SessionName);
+			NamedSessionUserMap.FindOrAdd(OpParams.LocalAccountId).AddUnique(OpParams.SessionName);
 
 			if (NewSessionEOSGSRef->SessionSettings.bPresenceEnabled)
 			{
-				FOnlineSessionIdHandle& PresenceSessionId = PresenceSessionsUserMap.FindOrAdd(OpParams.LocalUserId);
+				FOnlineSessionIdHandle& PresenceSessionId = PresenceSessionsUserMap.FindOrAdd(OpParams.LocalAccountId);
 				PresenceSessionId = NewSessionEOSGSRef->SessionId;
 			}
 
@@ -1074,12 +1074,12 @@ TOnlineAsyncOpHandle<FSendSingleSessionInviteImpl> FSessionsEOSGS::SendSingleSes
 	SendInviteOptions.ApiVersion = EOS_SESSIONS_SENDINVITE_API_LATEST;
 	static_assert(EOS_SESSIONS_SENDINVITE_API_LATEST == 1, "EOS_Sessions_SendInviteOptions updated, check new fields");
 
-	SendInviteOptions.LocalUserId = GetProductUserIdChecked(OpParams.LocalUserId);
+	SendInviteOptions.LocalUserId = GetProductUserIdChecked(OpParams.LocalAccountId);
 
 	const FTCHARToUTF8 SessionNameUtf8(*OpParams.SessionName.ToString());
 	SendInviteOptions.SessionName = SessionNameUtf8.Get();
 
-	SendInviteOptions.TargetUserId = GetProductUserIdChecked(OpParams.TargetUserId);
+	SendInviteOptions.TargetUserId = GetProductUserIdChecked(OpParams.TargetAccountId);
 
 	EOS_Async(EOS_Sessions_SendInvite, SessionsHandle, SendInviteOptions,
 		[this, WeakOp = Op->AsWeak()](const EOS_Sessions_SendInviteCallbackInfo* Result) mutable
@@ -1149,11 +1149,11 @@ TOnlineAsyncOpHandle<FLeaveSession> FSessionsEOSGS::LeaveSession(FLeaveSession::
 		{
 			TSharedRef<const FSession> FoundSession = GetSessionByNameResult.GetOkValue().Session;
 
-			NamedSessionUserMap.FindChecked(OpParams.LocalUserId).Remove(OpParams.SessionName);
+			NamedSessionUserMap.FindChecked(OpParams.LocalAccountId).Remove(OpParams.SessionName);
 
 			if (FoundSession->SessionSettings.bPresenceEnabled)
 			{
-				PresenceSessionsUserMap.Remove(OpParams.LocalUserId);
+				PresenceSessionsUserMap.Remove(OpParams.LocalAccountId);
 			}
 
 			ClearSessionByName(OpParams.SessionName);
@@ -1163,8 +1163,8 @@ TOnlineAsyncOpHandle<FLeaveSession> FSessionsEOSGS::LeaveSession(FLeaveSession::
 		Op.SetResult(FLeaveSession::Result{ });
 
 		FSessionLeft SessionLeftEvent;
-		SessionLeftEvent.LocalUserIds.Add(OpParams.LocalUserId);
-		//SessionLeftEvent.LocalUserIds.Append(OpParams.AdditionalLocalUsers); // TODO: Process multiple users in Session destruction
+		SessionLeftEvent.LocalAccountIds.Add(OpParams.LocalAccountId);
+		//SessionLeftEvent.LocalAccountIds.Append(OpParams.AdditionalLocalUsers); // TODO: Process multiple users in Session destruction
 		SessionEvents.OnSessionLeft.Broadcast(SessionLeftEvent);
 	})
 	.Enqueue(GetSerialQueue());
@@ -1225,12 +1225,12 @@ void FSessionsEOSGS::SetSessionSearchSessionId(FSessionSearchHandleEOSGS& Sessio
 	}
 }
 
-void FSessionsEOSGS::SetSessionSearchTargetId(FSessionSearchHandleEOSGS& SessionSearchHandle, const FAccountId& TargetUserId)
+void FSessionsEOSGS::SetSessionSearchTargetId(FSessionSearchHandleEOSGS& SessionSearchHandle, const FAccountId& TargetAccountId)
 {
 	EOS_SessionSearch_SetTargetUserIdOptions Options = { };
 	Options.ApiVersion = EOS_SESSIONSEARCH_SETTARGETUSERID_API_LATEST;
 	static_assert(EOS_SESSIONSEARCH_SETTARGETUSERID_API_LATEST == 1, "EOS_SessionSearch_SetTargetUserIdOptions updated, check new fields");
-	Options.TargetUserId = GetProductUserIdChecked(TargetUserId);
+	Options.TargetUserId = GetProductUserIdChecked(TargetAccountId);
 
 	EOS_EResult ResultCode = EOS_SessionSearch_SetTargetUserId(SessionSearchHandle.SearchHandle, &Options);
 	if (ResultCode != EOS_EResult::EOS_Success)
@@ -1303,8 +1303,8 @@ TOnlineAsyncOpHandle<FFindSessions> FSessionsEOSGS::FindSessions(FFindSessions::
 		}
 
 		// Before we start the search, we reset the cache
-		SearchResultsUserMap.FindOrAdd(OpParams.LocalUserId).Reset();
-		CurrentSessionSearchHandlesUserMap.Emplace(OpParams.LocalUserId, Op.AsShared());
+		SearchResultsUserMap.FindOrAdd(OpParams.LocalAccountId).Reset();
+		CurrentSessionSearchHandlesUserMap.Emplace(OpParams.LocalAccountId, Op.AsShared());
 
 		// We start preparing the search
 		EOS_Sessions_CreateSessionSearchOptions CreateSessionSearchOptions = { };
@@ -1323,7 +1323,7 @@ TOnlineAsyncOpHandle<FFindSessions> FSessionsEOSGS::FindSessions(FFindSessions::
 			return;
 		}
 
-		const TSharedRef<FSessionSearchHandleEOSGS>& CurrentSessionSearchHandleEOSGS = CurrentSessionSearchHandleEOSGSUserMap.Emplace(OpParams.LocalUserId, MakeShared<FSessionSearchHandleEOSGS>(SearchHandle));
+		const TSharedRef<FSessionSearchHandleEOSGS>& CurrentSessionSearchHandleEOSGS = CurrentSessionSearchHandleEOSGSUserMap.Emplace(OpParams.LocalAccountId, MakeShared<FSessionSearchHandleEOSGS>(SearchHandle));
 
 		// We write the search attributes
 		WriteSessionSearchHandle(*CurrentSessionSearchHandleEOSGS, OpParams);
@@ -1331,7 +1331,7 @@ TOnlineAsyncOpHandle<FFindSessions> FSessionsEOSGS::FindSessions(FFindSessions::
 		EOS_SessionSearch_FindOptions FindOptions = { };
 		FindOptions.ApiVersion = EOS_SESSIONSEARCH_FIND_API_LATEST;
 		static_assert(EOS_SESSIONSEARCH_FIND_API_LATEST == 2, "EOS_SessionSearch_FindOptions updated, check new fields");
-		FindOptions.LocalUserId = GetProductUserIdChecked(OpParams.LocalUserId);
+		FindOptions.LocalUserId = GetProductUserIdChecked(OpParams.LocalAccountId);
 
 		EOS_Async(EOS_SessionSearch_Find, CurrentSessionSearchHandleEOSGS->SearchHandle, FindOptions, MoveTemp(Promise));
 	})
@@ -1345,7 +1345,7 @@ TOnlineAsyncOpHandle<FFindSessions> FSessionsEOSGS::FindSessions(FFindSessions::
 
 		const FFindSessions::Params& OpParams = Op.GetParams();
 
-		const TSharedRef<FSessionSearchHandleEOSGS>& CurrentSessionSearchHandleEOSGS = CurrentSessionSearchHandleEOSGSUserMap.FindChecked(OpParams.LocalUserId);
+		const TSharedRef<FSessionSearchHandleEOSGS>& CurrentSessionSearchHandleEOSGS = CurrentSessionSearchHandleEOSGSUserMap.FindChecked(OpParams.LocalAccountId);
 
 		// For a successful session, we'll get the search results
 		EOS_SessionSearch_GetSearchResultCountOptions GetSearchResultCountOptions = { };
@@ -1369,7 +1369,7 @@ TOnlineAsyncOpHandle<FFindSessions> FSessionsEOSGS::FindSessions(FFindSessions::
 			if (CopySearchResultByIndexResult == EOS_EResult::EOS_Success)
 			{
 				TSharedRef<TPromise<TDefaultErrorResult<FBuildSessionFromDetailsHandle>>> BuildSessionPromise = MakeShared<TPromise<TDefaultErrorResult<FBuildSessionFromDetailsHandle>>>();
-				BuildSessionFromDetailsHandle({ OpParams.LocalUserId, MakeShared<FSessionDetailsHandleEOSGS>(SessionDetailsHandle) })
+				BuildSessionFromDetailsHandle({ OpParams.LocalAccountId, MakeShared<FSessionDetailsHandleEOSGS>(SessionDetailsHandle) })
 					.OnComplete([this, WeakOp = Op.AsWeak(), BuildSessionPromise](const TOnlineResult<FBuildSessionFromDetailsHandle>& Result) mutable
 				{
 					if (Result.IsOk())
@@ -1378,7 +1378,7 @@ TOnlineAsyncOpHandle<FFindSessions> FSessionsEOSGS::FindSessions(FFindSessions::
 						{
 							const TSharedRef<FSession>& Session = Result.GetOkValue().Session;
 
-							TArray<FOnlineSessionIdHandle>& SearchResults = SearchResultsUserMap.FindOrAdd(Result.GetOkValue().LocalUserId);
+							TArray<FOnlineSessionIdHandle>& SearchResults = SearchResultsUserMap.FindOrAdd(Result.GetOkValue().LocalAccountId);
 							SearchResults.Add(Session->SessionId);
 
 							AllSessionsById.Emplace(Session->SessionId, Session);
@@ -1419,12 +1419,12 @@ TOnlineAsyncOpHandle<FFindSessions> FSessionsEOSGS::FindSessions(FFindSessions::
 
 				const FFindSessions::Params& OpParams = StrongOp->GetParams();
 
-				TArray<FOnlineSessionIdHandle>& SearchResults = SearchResultsUserMap.FindChecked(OpParams.LocalUserId);
+				TArray<FOnlineSessionIdHandle>& SearchResults = SearchResultsUserMap.FindChecked(OpParams.LocalAccountId);
 
 				StrongOp->SetResult({ SearchResults });
 
-				CurrentSessionSearchHandlesUserMap.Remove(OpParams.LocalUserId);
-				CurrentSessionSearchHandleEOSGSUserMap.Remove(OpParams.LocalUserId);
+				CurrentSessionSearchHandlesUserMap.Remove(OpParams.LocalAccountId);
+				CurrentSessionSearchHandleEOSGSUserMap.Remove(OpParams.LocalAccountId);
 			}
 		});
 	});
@@ -1441,7 +1441,7 @@ TOnlineAsyncOpHandle<FJoinSession> FSessionsEOSGS::JoinSession(FJoinSession::Par
 	TOnlineAsyncOpRef<FJoinSession> Op = GetOp<FJoinSession>(MoveTemp(Params));
 	const FJoinSession::Params& OpParams = Op->GetParams();
 
-	TOnlineResult<FGetSessionById> GetSessionByIdResult = GetSessionById({ OpParams.LocalUserId, OpParams.SessionId });
+	TOnlineResult<FGetSessionById> GetSessionByIdResult = GetSessionById({ OpParams.LocalAccountId, OpParams.SessionId });
 	if (GetSessionByIdResult.IsError())
 	{
 		// If no result is found, the id might be expired, which we should notify
@@ -1494,7 +1494,7 @@ TOnlineAsyncOpHandle<FJoinSession> FSessionsEOSGS::JoinSession(FJoinSession::Par
 			return;
 		}
 
-		TOnlineResult<FGetSessionById> GetSessionByIdResult = GetSessionById({ OpParams.LocalUserId, OpParams.SessionId });
+		TOnlineResult<FGetSessionById> GetSessionByIdResult = GetSessionById({ OpParams.LocalAccountId, OpParams.SessionId });
 		if (GetSessionByIdResult.IsError())
 		{
 			// If no result is found, the id might be expired, which we should notify
@@ -1517,7 +1517,7 @@ TOnlineAsyncOpHandle<FJoinSession> FSessionsEOSGS::JoinSession(FJoinSession::Par
 
 		JoinSessionOptions.bPresenceEnabled = FoundSession->SessionSettings.bPresenceEnabled;
 
-		JoinSessionOptions.LocalUserId = GetProductUserIdChecked(OpParams.LocalUserId);
+		JoinSessionOptions.LocalUserId = GetProductUserIdChecked(OpParams.LocalAccountId);
 
 		const FTCHARToUTF8 SessionNameUtf8(*OpParams.SessionName.ToString());
 		JoinSessionOptions.SessionName = SessionNameUtf8.Get();
@@ -1537,7 +1537,7 @@ TOnlineAsyncOpHandle<FJoinSession> FSessionsEOSGS::JoinSession(FJoinSession::Par
 			return;
 		}
 		
-		TOnlineResult<FGetSessionById> GetSessionByIdResult = GetSessionById({ OpParams.LocalUserId, OpParams.SessionId });
+		TOnlineResult<FGetSessionById> GetSessionByIdResult = GetSessionById({ OpParams.LocalAccountId, OpParams.SessionId });
 		if (GetSessionByIdResult.IsError())
 		{
 			// If no result is found, the id might be expired, which we should notify
@@ -1554,36 +1554,36 @@ TOnlineAsyncOpHandle<FJoinSession> FSessionsEOSGS::JoinSession(FJoinSession::Par
 
 		LocalSessionsByName.Emplace(OpParams.SessionName, OpParams.SessionId);
 
-		NamedSessionUserMap.FindOrAdd(OpParams.LocalUserId).AddUnique(OpParams.SessionName);
+		NamedSessionUserMap.FindOrAdd(OpParams.LocalAccountId).AddUnique(OpParams.SessionName);
 
 		if (FoundSession->SessionSettings.bPresenceEnabled)
 		{
-			FOnlineSessionIdHandle& PresenceSessionId = PresenceSessionsUserMap.FindOrAdd(OpParams.LocalUserId);
+			FOnlineSessionIdHandle& PresenceSessionId = PresenceSessionsUserMap.FindOrAdd(OpParams.LocalAccountId);
 			PresenceSessionId = FoundSession->SessionId;
 		}
 
-		// After successfully joining a session, we'll remove all related invites if any are found
-		if (TMap<FOnlineSessionInviteIdHandle, TSharedRef<FSessionInvite>>* UserMap = SessionInvitesUserMap.Find(OpParams.LocalUserId))
-		{
-			TArray<FOnlineSessionInviteIdHandle> InviteIdsToRemove;
-			for (const TPair<FOnlineSessionInviteIdHandle, TSharedRef<FSessionInvite>>& Entry : *UserMap)
+			// After successfully joining a session, we'll remove all related invites if any are found
+			if (TMap<FOnlineSessionInviteIdHandle, TSharedRef<FSessionInvite>>* UserMap = SessionInvitesUserMap.Find(OpParams.LocalAccountId))
 			{
+				TArray<FOnlineSessionInviteIdHandle> InviteIdsToRemove;
+				for (const TPair<FOnlineSessionInviteIdHandle, TSharedRef<FSessionInvite>>& Entry : *UserMap)
+				{
 				if (Entry.Value->SessionId == FoundSession->SessionId)
+					{
+						InviteIdsToRemove.Add(Entry.Key);
+					}
+				}
+				for (const FOnlineSessionInviteIdHandle& InviteId : InviteIdsToRemove)
 				{
-					InviteIdsToRemove.Add(Entry.Key);
+					UserMap->Remove(InviteId);
 				}
 			}
-			for (const FOnlineSessionInviteIdHandle& InviteId : InviteIdsToRemove)
-			{
-				UserMap->Remove(InviteId);
-			}
-		}
 
 		Op.SetResult(FJoinSession::Result{ });
 
-		FSessionJoined Event = { { OpParams.LocalUserId }, FoundSession->SessionId };
+		FSessionJoined Event = { { OpParams.LocalAccountId }, FoundSession->SessionId };
 
-		SessionEvents.OnSessionJoined.Broadcast(Event);
+			SessionEvents.OnSessionJoined.Broadcast(Event);
 
 		// A successful join allows the client to server travel, after which RegisterPlayers will be called by the engine
 	})
@@ -1592,7 +1592,7 @@ TOnlineAsyncOpHandle<FJoinSession> FSessionsEOSGS::JoinSession(FJoinSession::Par
 	return Op->GetHandle();
 }
 
-TOnlineAsyncOpHandle<FBuildSessionFromDetailsHandle> FSessionsEOSGS::BuildSessionFromInvite(const FAccountId& LocalUserId, const FString& InInviteId)
+TOnlineAsyncOpHandle<FBuildSessionFromDetailsHandle> FSessionsEOSGS::BuildSessionFromInvite(const FAccountId& LocalAccountId, const FString& InInviteId)
 {
 	EOS_Sessions_CopySessionHandleByInviteIdOptions CopySessionHandleByInviteIdOptions = { };
 	CopySessionHandleByInviteIdOptions.ApiVersion = EOS_SESSIONS_COPYSESSIONHANDLEBYINVITEID_API_LATEST;
@@ -1605,7 +1605,7 @@ TOnlineAsyncOpHandle<FBuildSessionFromDetailsHandle> FSessionsEOSGS::BuildSessio
 	EOS_EResult CopySessionHandleByInviteIdResult = EOS_Sessions_CopySessionHandleByInviteId(SessionsHandle, &CopySessionHandleByInviteIdOptions, &SessionDetailsHandle);
 	if (CopySessionHandleByInviteIdResult == EOS_EResult::EOS_Success)
 	{
-		return BuildSessionFromDetailsHandle({ LocalUserId, MakeShared<FSessionDetailsHandleEOSGS>(SessionDetailsHandle) });
+		return BuildSessionFromDetailsHandle({ LocalAccountId, MakeShared<FSessionDetailsHandleEOSGS>(SessionDetailsHandle) });
 	}
 	else
 	{
@@ -1617,7 +1617,7 @@ TOnlineAsyncOpHandle<FBuildSessionFromDetailsHandle> FSessionsEOSGS::BuildSessio
 	}
 }
 
-TOnlineAsyncOpHandle<FBuildSessionFromDetailsHandle> FSessionsEOSGS::BuildSessionFromUIEvent(const FAccountId& LocalUserId, const EOS_UI_EventId& UIEventId)
+TOnlineAsyncOpHandle<FBuildSessionFromDetailsHandle> FSessionsEOSGS::BuildSessionFromUIEvent(const FAccountId& LocalAccountId, const EOS_UI_EventId& UIEventId)
 {
 	EOS_Sessions_CopySessionHandleByUiEventIdOptions CopySessionHandleByUiEventIdOptions = { };
 	CopySessionHandleByUiEventIdOptions.ApiVersion = EOS_SESSIONS_COPYSESSIONHANDLEBYUIEVENTID_API_LATEST;
@@ -1629,7 +1629,7 @@ TOnlineAsyncOpHandle<FBuildSessionFromDetailsHandle> FSessionsEOSGS::BuildSessio
 	EOS_EResult CopySessionHandleByUiEventIdResult = EOS_Sessions_CopySessionHandleByUiEventId(SessionsHandle, &CopySessionHandleByUiEventIdOptions, &SessionDetailsHandle);
 	if (CopySessionHandleByUiEventIdResult == EOS_EResult::EOS_Success)
 	{
-		return BuildSessionFromDetailsHandle({ LocalUserId, MakeShared<FSessionDetailsHandleEOSGS>(SessionDetailsHandle) });
+		return BuildSessionFromDetailsHandle({ LocalAccountId, MakeShared<FSessionDetailsHandleEOSGS>(SessionDetailsHandle) });
 	}
 	else
 	{
@@ -1712,13 +1712,13 @@ TOnlineAsyncOpHandle<FBuildSessionFromDetailsHandle> FSessionsEOSGS::BuildSessio
 
 	if (GetIdsResult.IsOk())
 	{
-		Services.Get<FAuthEOSGS>()->ResolveAccountIds(Params.LocalUserId, GetIdsResult.GetOkValue())
+		Services.Get<FAuthEOSGS>()->ResolveAccountIds(Params.LocalAccountId, GetIdsResult.GetOkValue())
 			.Next([this, WeakOp = Op->AsWeak(), Params = MoveTemp(Params)](TArray<FAccountId> ResolvedAccountIds) mutable
 		{
 			if (TOnlineAsyncOpPtr<FBuildSessionFromDetailsHandle> StrongOp = WeakOp.Pin())
 			{
 				// After all the ids are resolved, we can build the session safely
-				StrongOp->SetResult(FBuildSessionFromDetailsHandle::Result{ StrongOp->GetParams().LocalUserId, MakeShared<FSessionEOSGS>(Params.SessionDetailsHandleEOSGS->SessionDetailsHandle) });
+				StrongOp->SetResult(FBuildSessionFromDetailsHandle::Result{ StrongOp->GetParams().LocalAccountId, MakeShared<FSessionEOSGS>(Params.SessionDetailsHandleEOSGS->SessionDetailsHandle) });
 			}
 		});
 	}
@@ -1747,12 +1747,12 @@ TOnlineAsyncOpHandle<FSendSessionInvite> FSessionsEOSGS::SendSessionInvite(FSend
 		}
 
 		TArray<TFuture<TDefaultErrorResult<FSendSingleSessionInviteImpl>>> PendingSessionInvites;
-		for (const FAccountId& TargetUser : OpParams.TargetUsers)
+		for (const FAccountId& TargetAccountId : OpParams.TargetUsers)
 		{
 			FSendSingleSessionInviteImpl::Params FSendSingleSessionInviteParams;
-			FSendSingleSessionInviteParams.LocalUserId = OpParams.LocalUserId;
+			FSendSingleSessionInviteParams.LocalAccountId = OpParams.LocalAccountId;
 			FSendSingleSessionInviteParams.SessionName = OpParams.SessionName;
-			FSendSingleSessionInviteParams.TargetUserId = TargetUser;
+			FSendSingleSessionInviteParams.TargetAccountId = TargetAccountId;
 
 			TSharedRef<TPromise<TDefaultErrorResult<FSendSingleSessionInviteImpl>>> SessionInvitePromise = MakeShared<TPromise<TDefaultErrorResult<FSendSingleSessionInviteImpl>>>();
 			SendSingleSessionInviteImpl(MoveTemp(FSendSingleSessionInviteParams))
@@ -1813,7 +1813,7 @@ TOnlineAsyncOpHandle<FRejectSessionInvite> FSessionsEOSGS::RejectSessionInvite(F
 		RejectInviteOptions.ApiVersion = EOS_SESSIONS_REJECTINVITE_API_LATEST;
 		static_assert(EOS_SESSIONS_REJECTINVITE_API_LATEST == 1, "EOS_Sessions_RejectInviteOptions updated, check new fields");
 
-		RejectInviteOptions.LocalUserId = GetProductUserIdChecked(OpParams.LocalUserId);
+		RejectInviteOptions.LocalUserId = GetProductUserIdChecked(OpParams.LocalAccountId);
 
 		const FString& InviteIdStr = FOnlineSessionInviteIdRegistryEOSGS::Get().BasicRegistry.FindIdValue(OpParams.SessionInviteId);
 		const FTCHARToUTF8 InviteIdUtf8(*InviteIdStr);
@@ -1832,7 +1832,7 @@ TOnlineAsyncOpHandle<FRejectSessionInvite> FSessionsEOSGS::RejectSessionInvite(F
 
 		const FRejectSessionInvite::Params& OpParams = Op.GetParams();
 
-		if (TMap<FOnlineSessionInviteIdHandle, TSharedRef<FSessionInvite>>* UserMap = SessionInvitesUserMap.Find(OpParams.LocalUserId))
+		if (TMap<FOnlineSessionInviteIdHandle, TSharedRef<FSessionInvite>>* UserMap = SessionInvitesUserMap.Find(OpParams.LocalAccountId))
 		{
 			UserMap->Remove(OpParams.SessionInviteId);
 		}
@@ -1990,7 +1990,7 @@ void FSessionsEOSGS::AppendSessionToPacket(FNboSerializeToBuffer& Packet, const
 	using namespace NboSerializerEOSGSSvc;
 
 	SerializeToBuffer(Packet, Session);
-	SerializeToBuffer(Packet, Session.OwnerUserId);
+	SerializeToBuffer(Packet, Session.OwnerAccountId);
 	SerializeToBuffer(Packet, Session.SessionId);
 	SerializeToBuffer(Packet, Session.SessionSettings.RegisteredPlayers);
 	SerializeToBuffer(Packet, Session.SessionSettings.SessionMembers);
@@ -2002,7 +2002,7 @@ void FSessionsEOSGS::ReadSessionFromPacket(FNboSerializeFromBuffer& Packet, FSes
 	using namespace NboSerializerEOSGSSvc;
 
 	SerializeFromBuffer(Packet, Session);
-	SerializeFromBuffer(Packet, Session.OwnerUserId);
+	SerializeFromBuffer(Packet, Session.OwnerAccountId);
 	SerializeFromBuffer(Packet, Session.SessionId);
 	SerializeFromBuffer(Packet, Session.SessionSettings.RegisteredPlayers);
 	SerializeFromBuffer(Packet, Session.SessionSettings.SessionMembers);
diff --git a/Engine/Plugins/Online/Experimental/OnlineServicesEOSGS/Source/Private/Online/SocketSubsystemEOSUtils_OnlineServicesEOSGS.cpp b/Engine/Plugins/Online/Experimental/OnlineServicesEOSGS/Source/Private/Online/SocketSubsystemEOSUtils_OnlineServicesEOSGS.cpp
--- a/Engine/Plugins/Online/Experimental/OnlineServicesEOSGS/Source/Private/Online/SocketSubsystemEOSUtils_OnlineServicesEOSGS.cpp
+++ b/Engine/Plugins/Online/Experimental/OnlineServicesEOSGS/Source/Private/Online/SocketSubsystemEOSUtils_OnlineServicesEOSGS.cpp
@@ -62,7 +62,7 @@ FString FSocketSubsystemEOSUtils_OnlineServicesEOS::GetSessionId()
 		check(LobbiesEOS);
 
 		FGetJoinedLobbies::Params LobbiesParams;
-		LobbiesParams.LocalUserId = AuthOkValue->AccountInfo->AccountId;
+		LobbiesParams.LocalAccountId = AuthOkValue->AccountInfo->AccountId;
 		TOnlineResult<FGetJoinedLobbies> LobbiesResult = LobbiesEOS->GetJoinedLobbies(MoveTemp(LobbiesParams));
 		if (LobbiesResult.IsOk())
 		{
@@ -77,7 +77,7 @@ FString FSocketSubsystemEOSUtils_OnlineServicesEOS::GetSessionId()
 		}
 		else
 		{
-			UE_LOG(LogOnlineServices, Verbose, TEXT("[FSocketSubsystemEOSUtils_OnlineServicesEOS::GetSessionId] Unable to get joined lobbies for local user id [%s]. Error=[%s]"), *ToLogString(LobbiesParams.LocalUserId), *AuthResult.GetErrorValue().GetLogString(true));
+			UE_LOG(LogOnlineServices, Verbose, TEXT("[FSocketSubsystemEOSUtils_OnlineServicesEOS::GetSessionId] Unable to get joined lobbies for local user id [%s]. Error=[%s]"), *ToLogString(LobbiesParams.LocalAccountId), *AuthResult.GetErrorValue().GetLogString(true));
 		}
 	}
 	else
diff --git a/Engine/Plugins/Online/Experimental/OnlineServicesEOSGS/Source/Private/Online/StatsEOSGS.cpp b/Engine/Plugins/Online/Experimental/OnlineServicesEOSGS/Source/Private/Online/StatsEOSGS.cpp
--- a/Engine/Plugins/Online/Experimental/OnlineServicesEOSGS/Source/Private/Online/StatsEOSGS.cpp
+++ b/Engine/Plugins/Online/Experimental/OnlineServicesEOSGS/Source/Private/Online/StatsEOSGS.cpp
@@ -34,7 +34,7 @@ TOnlineAsyncOpHandle<FUpdateStats> FStatsEOSGS::UpdateStats(FUpdateStats::Params
 {
 	TOnlineAsyncOpRef<FUpdateStats> Op = GetOp<FUpdateStats>(MoveTemp(Params));
 
-	if (!Services.Get<FAuthEOSGS>()->IsLoggedIn(Op->GetParams().LocalUserId))
+	if (!Services.Get<FAuthEOSGS>()->IsLoggedIn(Op->GetParams().LocalAccountId))
 	{
 		Op->SetError(Errors::InvalidUser());
 		return Op->GetHandle();
@@ -70,8 +70,8 @@ TOnlineAsyncOpHandle<FUpdateStats> FStatsEOSGS::UpdateStats(FUpdateStats::Params
 			EOS_Stats_IngestStatOptions Options = { };
 			Options.ApiVersion = EOS_STATS_INGESTSTAT_API_LATEST;
 			static_assert(EOS_STATS_INGESTSTAT_API_LATEST == 3, "EOS_Stats_IngestStat updated, check new fields");
-			Options.LocalUserId = GetProductUserIdChecked(InAsyncOp.GetParams().LocalUserId);
-			Options.TargetUserId = GetProductUserIdChecked(UpdateUserStats.UserId);
+			Options.LocalUserId = GetProductUserIdChecked(InAsyncOp.GetParams().LocalAccountId);
+			Options.TargetUserId = GetProductUserIdChecked(UpdateUserStats.AccountId);
 			Options.Stats = EOSData.GetData();
 			Options.StatsCount = EOSData.Num();
 
@@ -102,7 +102,7 @@ TOnlineAsyncOpHandle<FUpdateStats> FStatsEOSGS::UpdateStats(FUpdateStats::Params
 namespace Private
 {
 
-void QueryStatsEOS(EOS_HStats StatsHandle, const FAccountId& LocalUserId, const FAccountId& TargetUserId, const TArray<FString>& StatNames, TPromise<const EOS_Stats_OnQueryStatsCompleteCallbackInfo*>&& Promise)
+void QueryStatsEOS(EOS_HStats StatsHandle, const FAccountId& LocalAccountId, const FAccountId& TargetAccountId, const TArray<FString>& StatNames, TPromise<const EOS_Stats_OnQueryStatsCompleteCallbackInfo*>&& Promise)
 {
 	EOS_Stats_QueryStatsOptions Options;
 	Options.ApiVersion = EOS_STATS_QUERYSTATS_API_LATEST;
@@ -122,8 +122,8 @@ void QueryStatsEOS(EOS_HStats StatsHandle, const FAccountId& LocalUserId, const
 
 	Options.StatNames = (const char**)EOSStatNames.GetData();
 	Options.StatNamesCount = StatNames.Num();
-	Options.LocalUserId = GetProductUserIdChecked(LocalUserId);
-	Options.TargetUserId = GetProductUserIdChecked(TargetUserId);
+	Options.LocalUserId = GetProductUserIdChecked(LocalAccountId);
+	Options.TargetUserId = GetProductUserIdChecked(TargetAccountId);
 
 	EOS_Async(EOS_Stats_QueryStats, StatsHandle, Options, MoveTemp(Promise));
 }
@@ -161,7 +161,7 @@ TOnlineAsyncOpHandle<FQueryStats> FStatsEOSGS::QueryStats(FQueryStats::Params&&
 {
 	TOnlineAsyncOpRef<FQueryStats> Op = GetOp<FQueryStats>(MoveTemp(Params));
 
-	if (!Services.Get<FAuthEOSGS>()->IsLoggedIn(Op->GetParams().LocalUserId))
+	if (!Services.Get<FAuthEOSGS>()->IsLoggedIn(Op->GetParams().LocalAccountId))
 	{
 		Op->SetError(Errors::InvalidUser());
 		return Op->GetHandle();
@@ -169,7 +169,7 @@ TOnlineAsyncOpHandle<FQueryStats> FStatsEOSGS::QueryStats(FQueryStats::Params&&
 
 	Op->Then([this](TOnlineAsyncOp<FQueryStats>& InAsyncOp, TPromise<const EOS_Stats_OnQueryStatsCompleteCallbackInfo*>&& Promise)
 	{
-		Private::QueryStatsEOS(StatsHandle, InAsyncOp.GetParams().LocalUserId, InAsyncOp.GetParams().TargetUserId, InAsyncOp.GetParams().StatNames, MoveTemp(Promise));
+		Private::QueryStatsEOS(StatsHandle, InAsyncOp.GetParams().LocalAccountId, InAsyncOp.GetParams().TargetAccountId, InAsyncOp.GetParams().StatNames, MoveTemp(Promise));
 	})
 	.Then([this](TOnlineAsyncOp<FQueryStats>& InAsyncOp, const EOS_Stats_OnQueryStatsCompleteCallbackInfo* Data)
 	{
@@ -184,7 +184,7 @@ TOnlineAsyncOpHandle<FQueryStats> FStatsEOSGS::QueryStats(FQueryStats::Params&&
 		Private::ReadStatsFromEOSResult(StatsHandle, Data, InAsyncOp.GetParams().StatNames, Result.Stats);
 
 		FUserStats UserStats;
-		UserStats.UserId = InAsyncOp.GetParams().TargetUserId;
+		UserStats.AccountId = InAsyncOp.GetParams().TargetAccountId;
 		UserStats.Stats = Result.Stats;
 		CacheUserStats(UserStats);
 
@@ -200,23 +200,23 @@ TOnlineAsyncOpHandle<FBatchQueryStats> FStatsEOSGS::BatchQueryStats(FBatchQueryS
 {
 	TOnlineAsyncOpRef<FBatchQueryStats> Op = GetOp<FBatchQueryStats>(MoveTemp(Params));
 
-	if (!Services.Get<FAuthEOSGS>()->IsLoggedIn(Op->GetParams().LocalUserId))
+	if (!Services.Get<FAuthEOSGS>()->IsLoggedIn(Op->GetParams().LocalAccountId))
 	{
 		Op->SetError(Errors::InvalidUser());
 		return Op->GetHandle();
 	}
 
-	if (Op->GetParams().TargetUserIds.IsEmpty() || Op->GetParams().StatNames.IsEmpty())
+	if (Op->GetParams().TargetAccountIds.IsEmpty() || Op->GetParams().StatNames.IsEmpty())
 	{
 		Op->SetError(Errors::InvalidParams());
 		return Op->GetHandle();
 	}
 
-	for (const FAccountId& TargetUserId : Op->GetParams().TargetUserIds)
+	for (const FAccountId& TargetAccountId : Op->GetParams().TargetAccountIds)
 	{
-		Op->Then([this, &TargetUserId](TOnlineAsyncOp<FBatchQueryStats>& InAsyncOp, TPromise<const EOS_Stats_OnQueryStatsCompleteCallbackInfo*>&& Promise)
+		Op->Then([this, &TargetAccountId](TOnlineAsyncOp<FBatchQueryStats>& InAsyncOp, TPromise<const EOS_Stats_OnQueryStatsCompleteCallbackInfo*>&& Promise)
 		{
-			Private::QueryStatsEOS(StatsHandle, InAsyncOp.GetParams().LocalUserId, TargetUserId, InAsyncOp.GetParams().StatNames, MoveTemp(Promise));
+			Private::QueryStatsEOS(StatsHandle, InAsyncOp.GetParams().LocalAccountId, TargetAccountId, InAsyncOp.GetParams().StatNames, MoveTemp(Promise));
 		})
 		.Then([this](TOnlineAsyncOp<FBatchQueryStats>& InAsyncOp, const EOS_Stats_OnQueryStatsCompleteCallbackInfo* Data)
 		{
diff --git a/Engine/Plugins/Online/Experimental/OnlineServicesEOSGS/Source/Private/Online/TitleFileEOSGS.cpp b/Engine/Plugins/Online/Experimental/OnlineServicesEOSGS/Source/Private/Online/TitleFileEOSGS.cpp
--- a/Engine/Plugins/Online/Experimental/OnlineServicesEOSGS/Source/Private/Online/TitleFileEOSGS.cpp
+++ b/Engine/Plugins/Online/Experimental/OnlineServicesEOSGS/Source/Private/Online/TitleFileEOSGS.cpp
@@ -41,7 +41,7 @@ TOnlineAsyncOpHandle<FTitleFileEnumerateFiles> FTitleFileEOSGS::EnumerateFiles(F
 	{
 		const FTitleFileEnumerateFiles::Params& Params = InAsyncOp.GetParams();
 
-		if (!Services.Get<FAuthEOSGS>()->IsLoggedIn(Params.LocalUserId))
+		if (!Services.Get<FAuthEOSGS>()->IsLoggedIn(Params.LocalAccountId))
 		{
 			InAsyncOp.SetError(Errors::InvalidUser());
 			Promise.EmplaceValue();
@@ -61,7 +61,7 @@ TOnlineAsyncOpHandle<FTitleFileEnumerateFiles> FTitleFileEOSGS::EnumerateFiles(F
 		EOS_TitleStorage_QueryFileListOptions Options = {};
 		Options.ApiVersion = EOS_TITLESTORAGE_QUERYFILELISTOPTIONS_API_LATEST;
 		static_assert(EOS_TITLESTORAGE_QUERYFILELISTOPTIONS_API_LATEST == 1, "EOS_TitleStorage_QueryFileListOptions updated, check new fields");
-		Options.LocalUserId = GetProductUserIdChecked(Params.LocalUserId);
+		Options.LocalUserId = GetProductUserIdChecked(Params.LocalAccountId);
 		Options.ListOfTagsCount = 1;
 		Options.ListOfTags = &SearchTagPtr;
 
@@ -78,7 +78,7 @@ TOnlineAsyncOpHandle<FTitleFileEnumerateFiles> FTitleFileEOSGS::EnumerateFiles(F
 			return;
 		}
 
-		const EOS_ProductUserId LocalUserPuid = GetProductUserIdChecked(Params.LocalUserId);
+		const EOS_ProductUserId LocalUserPuid = GetProductUserIdChecked(Params.LocalAccountId);
 
 		EOS_TitleStorage_GetFileMetadataCountOptions GetCountOptions;
 		GetCountOptions.ApiVersion = EOS_TITLESTORAGE_GETFILEMETADATACOUNTOPTIONS_API_LATEST;
@@ -122,7 +122,7 @@ TOnlineAsyncOpHandle<FTitleFileEnumerateFiles> FTitleFileEOSGS::EnumerateFiles(F
 
 TOnlineResult<FTitleFileGetEnumeratedFiles> FTitleFileEOSGS::GetEnumeratedFiles(FTitleFileGetEnumeratedFiles::Params&& Params)
 {
-	if (!Services.Get<FAuthEOSGS>()->IsLoggedIn(Params.LocalUserId))
+	if (!Services.Get<FAuthEOSGS>()->IsLoggedIn(Params.LocalAccountId))
 	{
 		return TOnlineResult<FTitleFileGetEnumeratedFiles>(Errors::InvalidUser());
 	}
@@ -167,7 +167,7 @@ TOnlineAsyncOpHandle<FTitleFileReadFile> FTitleFileEOSGS::ReadFile(FTitleFileRea
 	{
 		const FTitleFileReadFile::Params& Params = InAsyncOp.GetParams();
 
-		if (!Services.Get<FAuthEOSGS>()->IsLoggedIn(Params.LocalUserId))
+		if (!Services.Get<FAuthEOSGS>()->IsLoggedIn(Params.LocalAccountId))
 		{
 			InAsyncOp.SetError(Errors::InvalidUser());
 			Promise.EmplaceValue();
@@ -182,7 +182,7 @@ TOnlineAsyncOpHandle<FTitleFileReadFile> FTitleFileEOSGS::ReadFile(FTitleFileRea
 		EOS_TitleStorage_ReadFileOptions Options = {};
 		Options.ApiVersion = EOS_TITLESTORAGE_READFILEOPTIONS_API_LATEST;
 		static_assert(EOS_TITLESTORAGE_READFILEOPTIONS_API_LATEST == 1, "EOS_TitleStorage_ReadFileOptions updated, check new fields");
-		Options.LocalUserId = GetProductUserIdChecked(Params.LocalUserId);
+		Options.LocalUserId = GetProductUserIdChecked(Params.LocalAccountId);
 		Options.Filename = Utf8Filename.Get();
 		Options.ReadChunkLengthBytes = Config.ReadChunkLengthBytes;
 		Options.ReadFileDataCallback = &FTitleFileEOSGS::OnReadFileDataStatic;
diff --git a/Engine/Plugins/Online/Experimental/OnlineServicesEOSGS/Source/Private/Online/UserFileEOSGS.cpp b/Engine/Plugins/Online/Experimental/OnlineServicesEOSGS/Source/Private/Online/UserFileEOSGS.cpp
--- a/Engine/Plugins/Online/Experimental/OnlineServicesEOSGS/Source/Private/Online/UserFileEOSGS.cpp
+++ b/Engine/Plugins/Online/Experimental/OnlineServicesEOSGS/Source/Private/Online/UserFileEOSGS.cpp
@@ -44,7 +44,7 @@ TOnlineAsyncOpHandle<FUserFileEnumerateFiles> FUserFileEOSGS::EnumerateFiles(FUs
 	{
 		const FUserFileEnumerateFiles::Params& Params = Op.GetParams();
 
-		if (!Services.Get<FAuthEOSGS>()->IsLoggedIn(Params.LocalUserId))
+		if (!Services.Get<FAuthEOSGS>()->IsLoggedIn(Params.LocalAccountId))
 		{
 			Op.SetError(Errors::InvalidUser());
 			Promise.EmplaceValue();
@@ -54,7 +54,7 @@ TOnlineAsyncOpHandle<FUserFileEnumerateFiles> FUserFileEOSGS::EnumerateFiles(FUs
 		EOS_PlayerDataStorage_QueryFileListOptions Options = {};
 		Options.ApiVersion = EOS_PLAYERDATASTORAGE_QUERYFILELISTOPTIONS_API_LATEST;
 		static_assert(EOS_PLAYERDATASTORAGE_QUERYFILELISTOPTIONS_API_LATEST == 1, "EOS_PlayerDataStorage_QueryFileListOptions updated, check new fields");
-		Options.LocalUserId = GetProductUserIdChecked(Params.LocalUserId);
+		Options.LocalUserId = GetProductUserIdChecked(Params.LocalAccountId);
 
 		EOS_Async(EOS_PlayerDataStorage_QueryFileList, PlayerDataStorageHandle, Options, MoveTemp(Promise));
 	})
@@ -70,7 +70,7 @@ TOnlineAsyncOpHandle<FUserFileEnumerateFiles> FUserFileEOSGS::EnumerateFiles(FUs
 			return;
 		}
 
-		const EOS_ProductUserId LocalUserPuid = GetProductUserIdChecked(Params.LocalUserId);
+		const EOS_ProductUserId LocalUserPuid = GetProductUserIdChecked(Params.LocalAccountId);
 
 		EOS_PlayerDataStorage_GetFileMetadataCountOptions GetCountOptions;
 		GetCountOptions.ApiVersion = EOS_PLAYERDATASTORAGE_GETFILEMETADATACOUNTOPTIONS_API_LATEST;
@@ -113,7 +113,7 @@ TOnlineAsyncOpHandle<FUserFileEnumerateFiles> FUserFileEOSGS::EnumerateFiles(FUs
 			EOS_PlayerDataStorage_FileMetadata_Release(FileMetadata);
 		}
 
-		UserToFiles.Emplace(Params.LocalUserId, MoveTemp(NewEnumeratedFiles));
+		UserToFiles.Emplace(Params.LocalAccountId, MoveTemp(NewEnumeratedFiles));
 		Op.SetResult({});
 	})
 	.Enqueue(GetSerialQueue());
@@ -122,12 +122,12 @@ TOnlineAsyncOpHandle<FUserFileEnumerateFiles> FUserFileEOSGS::EnumerateFiles(FUs
 
 TOnlineResult<FUserFileGetEnumeratedFiles> FUserFileEOSGS::GetEnumeratedFiles(FUserFileGetEnumeratedFiles::Params&& Params)
 {
-	if (!Services.Get<FAuthEOSGS>()->IsLoggedIn(Params.LocalUserId))
+	if (!Services.Get<FAuthEOSGS>()->IsLoggedIn(Params.LocalAccountId))
 	{
 		return TOnlineResult<FUserFileGetEnumeratedFiles>(Errors::InvalidUser());
 	}
 
-	const TArray<FString>* Found = UserToFiles.Find(Params.LocalUserId);
+	const TArray<FString>* Found = UserToFiles.Find(Params.LocalAccountId);
 	if (!Found)
 	{
 		// Call EnumerateFiles first
@@ -167,7 +167,7 @@ TOnlineAsyncOpHandle<FUserFileReadFile> FUserFileEOSGS::ReadFile(FUserFileReadFi
 	{
 		const FUserFileReadFile::Params& Params = Op.GetParams();
 
-		if (!Services.Get<FAuthEOSGS>()->IsLoggedIn(Params.LocalUserId))
+		if (!Services.Get<FAuthEOSGS>()->IsLoggedIn(Params.LocalAccountId))
 		{
 			Op.SetError(Errors::InvalidUser());
 			Promise.EmplaceValue();
@@ -182,7 +182,7 @@ TOnlineAsyncOpHandle<FUserFileReadFile> FUserFileEOSGS::ReadFile(FUserFileReadFi
 		EOS_PlayerDataStorage_ReadFileOptions Options = {};
 		Options.ApiVersion = EOS_PLAYERDATASTORAGE_READFILEOPTIONS_API_LATEST;
 		static_assert(EOS_PLAYERDATASTORAGE_READFILEOPTIONS_API_LATEST == 1, "EOS_PlayerDataStorage_ReadFileOptions updated, check new fields");
-		Options.LocalUserId = GetProductUserIdChecked(Params.LocalUserId);
+		Options.LocalUserId = GetProductUserIdChecked(Params.LocalAccountId);
 		Options.Filename = (const char*)Utf8Filename.Get();
 		Options.ReadChunkLengthBytes = Config.ChunkLengthBytes;
 		Options.ReadFileDataCallback = &FUserFileEOSGS::OnReadFileDataStatic;
@@ -251,7 +251,7 @@ TOnlineAsyncOpHandle<FUserFileWriteFile> FUserFileEOSGS::WriteFile(FUserFileWrit
 	{
 		const FUserFileWriteFile::Params& Params = Op.GetParams();
 
-		if (!Services.Get<FAuthEOSGS>()->IsLoggedIn(Params.LocalUserId))
+		if (!Services.Get<FAuthEOSGS>()->IsLoggedIn(Params.LocalAccountId))
 		{
 			Op.SetError(Errors::InvalidUser());
 			Promise.EmplaceValue();
@@ -263,7 +263,7 @@ TOnlineAsyncOpHandle<FUserFileWriteFile> FUserFileEOSGS::WriteFile(FUserFileWrit
 		EOS_PlayerDataStorage_WriteFileOptions Options = {};
 		Options.ApiVersion = EOS_PLAYERDATASTORAGE_WRITEFILEOPTIONS_API_LATEST;
 		static_assert(EOS_PLAYERDATASTORAGE_WRITEFILEOPTIONS_API_LATEST == 1, "EOS_PlayerDataStorage_WriteFileOptions updated, check new fields");
-		Options.LocalUserId = GetProductUserIdChecked(Params.LocalUserId);
+		Options.LocalUserId = GetProductUserIdChecked(Params.LocalAccountId);
 		Options.Filename = (const char*)Utf8Filename.Get();
 		Options.ChunkLengthBytes = Config.ChunkLengthBytes;
 		Options.WriteFileDataCallback = &FUserFileEOSGS::OnWriteFileDataStatic;
@@ -310,7 +310,7 @@ TOnlineAsyncOpHandle<FUserFileCopyFile> FUserFileEOSGS::CopyFile(FUserFileCopyFi
 	{
 		const FUserFileCopyFile::Params& Params = Op.GetParams();
 
-		if (!Services.Get<FAuthEOSGS>()->IsLoggedIn(Params.LocalUserId))
+		if (!Services.Get<FAuthEOSGS>()->IsLoggedIn(Params.LocalAccountId))
 		{
 			Op.SetError(Errors::InvalidUser());
 			Promise.EmplaceValue();
@@ -323,7 +323,7 @@ TOnlineAsyncOpHandle<FUserFileCopyFile> FUserFileEOSGS::CopyFile(FUserFileCopyFi
 		EOS_PlayerDataStorage_DuplicateFileOptions Options = {};
 		Options.ApiVersion = EOS_PLAYERDATASTORAGE_DUPLICATEFILEOPTIONS_API_LATEST;
 		static_assert(EOS_PLAYERDATASTORAGE_DUPLICATEFILEOPTIONS_API_LATEST == 1, "EOS_PlayerDataStorage_DuplicateFileOptions updated, check new fields");
-		Options.LocalUserId = GetProductUserIdChecked(Params.LocalUserId);
+		Options.LocalUserId = GetProductUserIdChecked(Params.LocalAccountId);
 		Options.SourceFilename = (const char*)Utf8SourceFilename.Get();
 		Options.DestinationFilename = (const char*)Utf8TargetFilename.Get();
 
@@ -360,7 +360,7 @@ TOnlineAsyncOpHandle<FUserFileDeleteFile> FUserFileEOSGS::DeleteFile(FUserFileDe
 	{
 		const FUserFileDeleteFile::Params& Params = Op.GetParams();
 
-		if (!Services.Get<FAuthEOSGS>()->IsLoggedIn(Params.LocalUserId))
+		if (!Services.Get<FAuthEOSGS>()->IsLoggedIn(Params.LocalAccountId))
 		{
 			Op.SetError(Errors::InvalidUser());
 			Promise.EmplaceValue();
@@ -372,7 +372,7 @@ TOnlineAsyncOpHandle<FUserFileDeleteFile> FUserFileEOSGS::DeleteFile(FUserFileDe
 		EOS_PlayerDataStorage_DeleteFileOptions Options = {};
 		Options.ApiVersion = EOS_PLAYERDATASTORAGE_DELETEFILEOPTIONS_API_LATEST;
 		static_assert(EOS_PLAYERDATASTORAGE_DELETEFILEOPTIONS_API_LATEST == 1, "EOS_PlayerDataStorage_DeleteFileOptions updated, check new fields");
-		Options.LocalUserId = GetProductUserIdChecked(Params.LocalUserId);
+		Options.LocalUserId = GetProductUserIdChecked(Params.LocalAccountId);
 		Options.Filename = (const char*)Utf8Filename.Get();
 
 		EOS_Async(EOS_PlayerDataStorage_DeleteFile, PlayerDataStorageHandle, Options, MoveTemp(Promise));
diff --git a/Engine/Plugins/Online/Experimental/OnlineServicesEOSGS/Source/Public/Online/AuthEOSGS.h b/Engine/Plugins/Online/Experimental/OnlineServicesEOSGS/Source/Public/Online/AuthEOSGS.h
--- a/Engine/Plugins/Online/Experimental/OnlineServicesEOSGS/Source/Public/Online/AuthEOSGS.h
+++ b/Engine/Plugins/Online/Experimental/OnlineServicesEOSGS/Source/Public/Online/AuthEOSGS.h
@@ -169,7 +169,7 @@ public:
 	virtual ~FAccountInfoRegistryEOS() = default;
 
 	TSharedPtr<FAccountInfoEOS> Find(FPlatformUserId PlatformUserId) const;
-	TSharedPtr<FAccountInfoEOS> Find(FAccountId AccountIdHandle) const;
+	TSharedPtr<FAccountInfoEOS> Find(FAccountId AccountId) const;
 	TSharedPtr<FAccountInfoEOS> Find(EOS_EpicAccountId EpicAccountId) const;
 	TSharedPtr<FAccountInfoEOS> Find(EOS_ProductUserId ProductUserId) const;
 
@@ -208,8 +208,8 @@ public:
 	// End IAuth
 
 	// Begin FAuthEOSGS
-	virtual TFuture<FAccountId> ResolveAccountId(const FAccountId& LocalUserId, const EOS_ProductUserId ProductUserId);
-	virtual TFuture<TArray<FAccountId>> ResolveAccountIds(const FAccountId& LocalUserId, const TArray<EOS_ProductUserId>& ProductUserIds);
+	virtual TFuture<FAccountId> ResolveAccountId(const FAccountId& LocalAccountId, const EOS_ProductUserId ProductUserId);
+	virtual TFuture<TArray<FAccountId>> ResolveAccountIds(const FAccountId& LocalAccountId, const TArray<EOS_ProductUserId>& ProductUserIds);
 	virtual TFunction<TFuture<FAccountId>(FOnlineAsyncOp& InAsyncOp, const EOS_ProductUserId& ProductUserId)> ResolveProductIdFn();
 	virtual TFunction<TFuture<TArray<FAccountId>>(FOnlineAsyncOp& InAsyncOp, const TArray<EOS_ProductUserId>& ProductUserIds)> ResolveProductIdsFn();
 	// End FAuthEOSGS
diff --git a/Engine/Plugins/Online/Experimental/OnlineServicesEOSGS/Source/Public/Online/LobbiesEOSGS.h b/Engine/Plugins/Online/Experimental/OnlineServicesEOSGS/Source/Public/Online/LobbiesEOSGS.h
--- a/Engine/Plugins/Online/Experimental/OnlineServicesEOSGS/Source/Public/Online/LobbiesEOSGS.h
+++ b/Engine/Plugins/Online/Experimental/OnlineServicesEOSGS/Source/Public/Online/LobbiesEOSGS.h
@@ -32,7 +32,7 @@ struct FLobbiesJoinLobbyImpl
 		TSharedPtr<FLobbyDataEOS> LobbyData;
 
 		// The local user agent which will perform the action.
-		FAccountId LocalUserId;
+		FAccountId LocalAccountId;
 
 		// The local name for the lobby.
 		FName LocalName;
@@ -56,7 +56,7 @@ struct FLobbiesJoinLobbyMemberImpl
 		TSharedPtr<FLobbyDataEOS> LobbyData;
 
 		// The local user agent which will perform the action.
-		FAccountId LocalUserId;
+		FAccountId LocalAccountId;
 
 		// Initial attributes.
 		TMap<FLobbyAttributeId, FLobbyVariant> Attributes;
@@ -77,7 +77,7 @@ struct FLobbiesLeaveLobbyImpl
 		TSharedPtr<FLobbyDataEOS> LobbyData;
 
 		// The local user agent which will perform the action.
-		FAccountId LocalUserId;
+		FAccountId LocalAccountId;
 	};
 
 	struct Result
@@ -95,7 +95,7 @@ struct FLobbiesDestroyLobbyImpl
 		TSharedPtr<FLobbyDataEOS> LobbyData;
 
 		// The local user agent which will perform the action.
-		FAccountId LocalUserId;
+		FAccountId LocalAccountId;
 	};
 
 	struct Result
@@ -113,10 +113,10 @@ struct FLobbiesInviteLobbyMemberImpl
 		TSharedPtr<FLobbyDataEOS> LobbyData;
 
 		// The local user agent which will perform the action.
-		FAccountId LocalUserId;
+		FAccountId LocalAccountId;
 
 		// The target user for the invitation.
-		FAccountId TargetUserId;
+		FAccountId TargetAccountId;
 	};
 
 	struct Result
@@ -131,7 +131,7 @@ struct FLobbiesDeclineLobbyInvitationImpl
 	struct Params
 	{
 		// The local user agent which will perform the action.
-		FAccountId LocalUserId;
+		FAccountId LocalAccountId;
 
 		// Id of the lobby for which the invitations will be declined.
 		FOnlineLobbyIdHandle LobbyId;
@@ -152,10 +152,10 @@ struct FLobbiesKickLobbyMemberImpl
 		TSharedPtr<FLobbyDataEOS> LobbyData;
 
 		// The local user agent which will perform the action.
-		FAccountId LocalUserId;
+		FAccountId LocalAccountId;
 
 		// The target user to be kicked.
-		FAccountId TargetUserId;
+		FAccountId TargetAccountId;
 	};
 
 	struct Result
@@ -173,10 +173,10 @@ struct FLobbiesPromoteLobbyMemberImpl
 		TSharedPtr<FLobbyDataEOS> LobbyData;
 
 		// The local user agent which will perform the action.
-		FAccountId LocalUserId;
+		FAccountId LocalAccountId;
 
 		// The target user to be promoted to owner.
-		FAccountId TargetUserId;
+		FAccountId TargetAccountId;
 	};
 
 	struct Result
@@ -194,7 +194,7 @@ struct FLobbiesModifyLobbyDataImpl
 		TSharedPtr<FLobbyDataEOS> LobbyData;
 
 		// The local user agent which will perform the action.
-		FAccountId LocalUserId;
+		FAccountId LocalAccountId;
 
 		// The changes to apply to the lobby data.
 		TSharedPtr<FClientLobbyDataChanges> Changes;
@@ -218,7 +218,7 @@ struct FLobbiesModifyLobbyMemberDataImpl
 		TSharedPtr<FLobbyDataEOS> LobbyData;
 
 		// The local user agent which will perform the action.
-		FAccountId LocalUserId;
+		FAccountId LocalAccountId;
 
 		// The changes to apply to the lobby member data.
 		TSharedPtr<FClientLobbyMemberDataChanges> Changes;
@@ -293,8 +293,8 @@ protected:
 	void RegisterHandlers();
 	void UnregisterHandlers();
 
-	void AddActiveLobby(FAccountId LocalUserId, const TSharedRef<FLobbyDataEOS>& LobbyData);
-	void RemoveActiveLobby(FAccountId LocalUserId, const TSharedRef<FLobbyDataEOS>& LobbyData);
+	void AddActiveLobby(FAccountId LocalAccountId, const TSharedRef<FLobbyDataEOS>& LobbyData);
+	void RemoveActiveLobby(FAccountId LocalAccountId, const TSharedRef<FLobbyDataEOS>& LobbyData);
 
 	// Todo: store list of invites per lobby.
 	void AddActiveInvite(const TSharedRef<FLobbyInviteDataEOS>& Invite);
@@ -334,7 +334,7 @@ namespace Meta {
 
 BEGIN_ONLINE_STRUCT_META(FLobbiesJoinLobbyImpl::Params)
 	ONLINE_STRUCT_FIELD(FLobbiesJoinLobbyImpl::Params, LobbyData),
-	ONLINE_STRUCT_FIELD(FLobbiesJoinLobbyImpl::Params, LocalUserId),
+	ONLINE_STRUCT_FIELD(FLobbiesJoinLobbyImpl::Params, LocalAccountId),
 	ONLINE_STRUCT_FIELD(FLobbiesJoinLobbyImpl::Params, LocalName),
 	ONLINE_STRUCT_FIELD(FLobbiesJoinLobbyImpl::Params, LocalUsers)
 END_ONLINE_STRUCT_META()
@@ -344,7 +344,7 @@ END_ONLINE_STRUCT_META()
 
 BEGIN_ONLINE_STRUCT_META(FLobbiesJoinLobbyMemberImpl::Params)
 	ONLINE_STRUCT_FIELD(FLobbiesJoinLobbyMemberImpl::Params, LobbyData),
-	ONLINE_STRUCT_FIELD(FLobbiesJoinLobbyMemberImpl::Params, LocalUserId),
+	ONLINE_STRUCT_FIELD(FLobbiesJoinLobbyMemberImpl::Params, LocalAccountId),
 	ONLINE_STRUCT_FIELD(FLobbiesJoinLobbyMemberImpl::Params, Attributes)
 END_ONLINE_STRUCT_META()
 
@@ -353,7 +353,7 @@ END_ONLINE_STRUCT_META()
 
 BEGIN_ONLINE_STRUCT_META(FLobbiesLeaveLobbyImpl::Params)
 	ONLINE_STRUCT_FIELD(FLobbiesLeaveLobbyImpl::Params, LobbyData),
-	ONLINE_STRUCT_FIELD(FLobbiesLeaveLobbyImpl::Params, LocalUserId)
+	ONLINE_STRUCT_FIELD(FLobbiesLeaveLobbyImpl::Params, LocalAccountId)
 END_ONLINE_STRUCT_META()
 
 BEGIN_ONLINE_STRUCT_META(FLobbiesLeaveLobbyImpl::Result)
@@ -361,7 +361,7 @@ END_ONLINE_STRUCT_META()
 
 BEGIN_ONLINE_STRUCT_META(FLobbiesDestroyLobbyImpl::Params)
 	ONLINE_STRUCT_FIELD(FLobbiesDestroyLobbyImpl::Params, LobbyData),
-	ONLINE_STRUCT_FIELD(FLobbiesDestroyLobbyImpl::Params, LocalUserId)
+	ONLINE_STRUCT_FIELD(FLobbiesDestroyLobbyImpl::Params, LocalAccountId)
 END_ONLINE_STRUCT_META()
 
 BEGIN_ONLINE_STRUCT_META(FLobbiesDestroyLobbyImpl::Result)
@@ -369,15 +369,15 @@ END_ONLINE_STRUCT_META()
 
 BEGIN_ONLINE_STRUCT_META(FLobbiesInviteLobbyMemberImpl::Params)
 	ONLINE_STRUCT_FIELD(FLobbiesInviteLobbyMemberImpl::Params, LobbyData),
-	ONLINE_STRUCT_FIELD(FLobbiesInviteLobbyMemberImpl::Params, LocalUserId),
-	ONLINE_STRUCT_FIELD(FLobbiesInviteLobbyMemberImpl::Params, TargetUserId)
+	ONLINE_STRUCT_FIELD(FLobbiesInviteLobbyMemberImpl::Params, LocalAccountId),
+	ONLINE_STRUCT_FIELD(FLobbiesInviteLobbyMemberImpl::Params, TargetAccountId)
 END_ONLINE_STRUCT_META()
 
 BEGIN_ONLINE_STRUCT_META(FLobbiesInviteLobbyMemberImpl::Result)
 END_ONLINE_STRUCT_META()
 
 BEGIN_ONLINE_STRUCT_META(FLobbiesDeclineLobbyInvitationImpl::Params)
-	ONLINE_STRUCT_FIELD(FLobbiesDeclineLobbyInvitationImpl::Params, LocalUserId),
+	ONLINE_STRUCT_FIELD(FLobbiesDeclineLobbyInvitationImpl::Params, LocalAccountId),
 	ONLINE_STRUCT_FIELD(FLobbiesDeclineLobbyInvitationImpl::Params, LobbyId)
 END_ONLINE_STRUCT_META()
 
@@ -386,8 +386,8 @@ END_ONLINE_STRUCT_META()
 
 BEGIN_ONLINE_STRUCT_META(FLobbiesKickLobbyMemberImpl::Params)
 	ONLINE_STRUCT_FIELD(FLobbiesKickLobbyMemberImpl::Params, LobbyData),
-	ONLINE_STRUCT_FIELD(FLobbiesKickLobbyMemberImpl::Params, LocalUserId),
-	ONLINE_STRUCT_FIELD(FLobbiesKickLobbyMemberImpl::Params, TargetUserId)
+	ONLINE_STRUCT_FIELD(FLobbiesKickLobbyMemberImpl::Params, LocalAccountId),
+	ONLINE_STRUCT_FIELD(FLobbiesKickLobbyMemberImpl::Params, TargetAccountId)
 END_ONLINE_STRUCT_META()
 
 BEGIN_ONLINE_STRUCT_META(FLobbiesKickLobbyMemberImpl::Result)
@@ -395,8 +395,8 @@ END_ONLINE_STRUCT_META()
 
 BEGIN_ONLINE_STRUCT_META(FLobbiesPromoteLobbyMemberImpl::Params)
 	ONLINE_STRUCT_FIELD(FLobbiesPromoteLobbyMemberImpl::Params, LobbyData),
-	ONLINE_STRUCT_FIELD(FLobbiesPromoteLobbyMemberImpl::Params, LocalUserId),
-	ONLINE_STRUCT_FIELD(FLobbiesPromoteLobbyMemberImpl::Params, TargetUserId)
+	ONLINE_STRUCT_FIELD(FLobbiesPromoteLobbyMemberImpl::Params, LocalAccountId),
+	ONLINE_STRUCT_FIELD(FLobbiesPromoteLobbyMemberImpl::Params, TargetAccountId)
 END_ONLINE_STRUCT_META()
 
 BEGIN_ONLINE_STRUCT_META(FLobbiesPromoteLobbyMemberImpl::Result)
@@ -404,7 +404,7 @@ END_ONLINE_STRUCT_META()
 
 BEGIN_ONLINE_STRUCT_META(FLobbiesModifyLobbyDataImpl::Params)
 	ONLINE_STRUCT_FIELD(FLobbiesModifyLobbyDataImpl::Params, LobbyData),
-	ONLINE_STRUCT_FIELD(FLobbiesModifyLobbyDataImpl::Params, LocalUserId)
+	ONLINE_STRUCT_FIELD(FLobbiesModifyLobbyDataImpl::Params, LocalAccountId)
 END_ONLINE_STRUCT_META()
 
 BEGIN_ONLINE_STRUCT_META(FLobbiesModifyLobbyDataImpl::Result)
@@ -412,7 +412,7 @@ END_ONLINE_STRUCT_META()
 
 BEGIN_ONLINE_STRUCT_META(FLobbiesModifyLobbyMemberDataImpl::Params)
 	ONLINE_STRUCT_FIELD(FLobbiesModifyLobbyMemberDataImpl::Params, LobbyData),
-	ONLINE_STRUCT_FIELD(FLobbiesModifyLobbyMemberDataImpl::Params, LocalUserId)
+	ONLINE_STRUCT_FIELD(FLobbiesModifyLobbyMemberDataImpl::Params, LocalAccountId)
 END_ONLINE_STRUCT_META()
 
 BEGIN_ONLINE_STRUCT_META(FLobbiesModifyLobbyMemberDataImpl::Result)
diff --git a/Engine/Plugins/Online/Experimental/OnlineServicesEOSGS/Source/Public/Online/OnlineIdEOSGS.h b/Engine/Plugins/Online/Experimental/OnlineServicesEOSGS/Source/Public/Online/OnlineIdEOSGS.h
--- a/Engine/Plugins/Online/Experimental/OnlineServicesEOSGS/Source/Public/Online/OnlineIdEOSGS.h
+++ b/Engine/Plugins/Online/Experimental/OnlineServicesEOSGS/Source/Public/Online/OnlineIdEOSGS.h
@@ -16,7 +16,7 @@ class IOnlineAccountIdRegistryEOSGS : public IOnlineAccountIdRegistry
 {
 public:
 	virtual FAccountId FindAccountId(EOS_ProductUserId ProductUserId) const = 0;
-	virtual EOS_ProductUserId GetProductUserId(const FAccountId& Handle) const = 0;
+	virtual EOS_ProductUserId GetProductUserId(const FAccountId& AccountId) const = 0;
 };
 
 /**
@@ -31,12 +31,12 @@ public:
 
 	// Begin IOnlineAccountIdRegistryEOSGS
 	virtual FAccountId FindAccountId(const EOS_ProductUserId ProductUserId) const override;
-	virtual EOS_ProductUserId GetProductUserId(const FAccountId& Handle) const override;
+	virtual EOS_ProductUserId GetProductUserId(const FAccountId& AccountId) const override;
 	// End IOnlineAccountIdRegistryEOSGS
 
 	// Begin IOnlineAccountIdRegistry
-	virtual FString ToLogString(const FAccountId& Handle) const override;
-	virtual TArray<uint8> ToReplicationData(const FAccountId& Handle) const override;
+	virtual FString ToLogString(const FAccountId& AccountId) const override;
+	virtual TArray<uint8> ToReplicationData(const FAccountId& AccountId) const override;
 	virtual FAccountId FromReplicationData(const TArray<uint8>& ReplicationData) override;
 	// End IOnlineAccountIdRegistry
 
@@ -52,8 +52,8 @@ private:
 	TOnlineBasicAccountIdRegistry<EOS_ProductUserId> Registry;
 };
 
-EOS_ProductUserId ONLINESERVICESEOSGS_API GetProductUserId(const FAccountId& Handle);
-EOS_ProductUserId ONLINESERVICESEOSGS_API GetProductUserIdChecked(const FAccountId& Handle);
+EOS_ProductUserId ONLINESERVICESEOSGS_API GetProductUserId(const FAccountId& AccountId);
+EOS_ProductUserId ONLINESERVICESEOSGS_API GetProductUserIdChecked(const FAccountId& AccountId);
 FAccountId ONLINESERVICESEOSGS_API FindAccountId(const EOS_ProductUserId EpicAccountId);
 FAccountId ONLINESERVICESEOSGS_API FindAccountIdChecked(const EOS_ProductUserId EpicAccountId);
 
diff --git a/Engine/Plugins/Online/Experimental/OnlineServicesEOSGS/Source/Public/Online/SessionsEOSGS.h b/Engine/Plugins/Online/Experimental/OnlineServicesEOSGS/Source/Public/Online/SessionsEOSGS.h
--- a/Engine/Plugins/Online/Experimental/OnlineServicesEOSGS/Source/Public/Online/SessionsEOSGS.h
+++ b/Engine/Plugins/Online/Experimental/OnlineServicesEOSGS/Source/Public/Online/SessionsEOSGS.h
@@ -140,11 +140,11 @@ struct FSendSingleSessionInviteImpl
 
 	struct Params
 	{
-		FAccountId LocalUserId;
+		FAccountId LocalAccountId;
 
 		FName SessionName;
 
-		FAccountId TargetUserId;
+		FAccountId TargetAccountId;
 	};
 
 	struct Result
@@ -159,7 +159,7 @@ struct FBuildSessionFromDetailsHandle
 	struct Params
 	{
 		/** User which will drive the id resolution */
-		FAccountId LocalUserId;
+		FAccountId LocalAccountId;
 
 		/** EOS session details handle used to extract the data */
 		TSharedRef<FSessionDetailsHandleEOSGS> SessionDetailsHandleEOSGS;
@@ -168,7 +168,7 @@ struct FBuildSessionFromDetailsHandle
 	struct Result
 	{
 		/** User which started the resolution operation */
-		FAccountId LocalUserId;
+		FAccountId LocalAccountId;
 
 		/** Session built from the details handle */
 		TSharedRef<FSession> Session;
@@ -218,7 +218,7 @@ protected:
 	void SetSessionSearchMaxResults(FSessionSearchHandleEOSGS& SessionSearchHandle, uint32 MaxResults);
 	void SetSessionSearchParameters(FSessionSearchHandleEOSGS& SessionSearchHandle, TArray<FFindSessionsSearchFilter> Filters);
 	void SetSessionSearchSessionId(FSessionSearchHandleEOSGS& SessionSearchHandle, const FOnlineSessionIdHandle& SessionId);
-	void SetSessionSearchTargetId(FSessionSearchHandleEOSGS& SessionSearchHandle, const FAccountId& TargetUserId);
+	void SetSessionSearchTargetId(FSessionSearchHandleEOSGS& SessionSearchHandle, const FAccountId& TargetAccountId);
 
 	/**
 	 * Writes all values in the passed SessionSettings to the SessionModificationHandle
@@ -302,21 +302,21 @@ BEGIN_ONLINE_STRUCT_META(FUpdateSessionImpl::Result)
 END_ONLINE_STRUCT_META()
 
 BEGIN_ONLINE_STRUCT_META(FSendSingleSessionInviteImpl::Params)
-	ONLINE_STRUCT_FIELD(FSendSingleSessionInviteImpl::Params, LocalUserId),
+	ONLINE_STRUCT_FIELD(FSendSingleSessionInviteImpl::Params, LocalAccountId),
 	ONLINE_STRUCT_FIELD(FSendSingleSessionInviteImpl::Params, SessionName),
-	ONLINE_STRUCT_FIELD(FSendSingleSessionInviteImpl::Params, TargetUserId)
+	ONLINE_STRUCT_FIELD(FSendSingleSessionInviteImpl::Params, TargetAccountId)
 END_ONLINE_STRUCT_META()
 
 BEGIN_ONLINE_STRUCT_META(FSendSingleSessionInviteImpl::Result)
 END_ONLINE_STRUCT_META()
 
 BEGIN_ONLINE_STRUCT_META(FBuildSessionFromDetailsHandle::Params)
-	ONLINE_STRUCT_FIELD(FBuildSessionFromDetailsHandle::Params, LocalUserId),
+	ONLINE_STRUCT_FIELD(FBuildSessionFromDetailsHandle::Params, LocalAccountId),
 	ONLINE_STRUCT_FIELD(FBuildSessionFromDetailsHandle::Params, SessionDetailsHandleEOSGS)
 END_ONLINE_STRUCT_META()
 
 BEGIN_ONLINE_STRUCT_META(FBuildSessionFromDetailsHandle::Result)
-	ONLINE_STRUCT_FIELD(FBuildSessionFromDetailsHandle::Result, LocalUserId),
+	ONLINE_STRUCT_FIELD(FBuildSessionFromDetailsHandle::Result, LocalAccountId),
 	ONLINE_STRUCT_FIELD(FBuildSessionFromDetailsHandle::Result, Session)
 END_ONLINE_STRUCT_META()
 
diff --git a/Engine/Plugins/Online/Experimental/OnlineServicesNull/Source/Private/Online/AchievementsNull.cpp b/Engine/Plugins/Online/Experimental/OnlineServicesNull/Source/Private/Online/AchievementsNull.cpp
--- a/Engine/Plugins/Online/Experimental/OnlineServicesNull/Source/Private/Online/AchievementsNull.cpp
+++ b/Engine/Plugins/Online/Experimental/OnlineServicesNull/Source/Private/Online/AchievementsNull.cpp
@@ -73,7 +73,7 @@ TOnlineAsyncOpHandle<FQueryAchievementDefinitions> FAchievementsNull::QueryAchie
 {
 	TOnlineAsyncOpRef<FQueryAchievementDefinitions> Op = GetOp<FQueryAchievementDefinitions>(MoveTemp(Params));
 
-	if (!Services.Get<FAuthNull>()->IsLoggedIn(Op->GetParams().LocalUserId))
+	if (!Services.Get<FAuthNull>()->IsLoggedIn(Op->GetParams().LocalAccountId))
 	{
 		Op->SetError(Errors::InvalidUser());
 		return Op->GetHandle();
@@ -90,7 +90,7 @@ TOnlineAsyncOpHandle<FQueryAchievementDefinitions> FAchievementsNull::QueryAchie
 
 TOnlineResult<FGetAchievementIds> FAchievementsNull::GetAchievementIds(FGetAchievementIds::Params&& Params)
 {
-	if (!Services.Get<FAuthNull>()->IsLoggedIn(Params.LocalUserId))
+	if (!Services.Get<FAuthNull>()->IsLoggedIn(Params.LocalAccountId))
 	{
 		return TOnlineResult<FGetAchievementIds>(Errors::InvalidUser());
 	}
@@ -108,7 +108,7 @@ TOnlineResult<FGetAchievementIds> FAchievementsNull::GetAchievementIds(FGetAchie
 
 TOnlineResult<FGetAchievementDefinition> FAchievementsNull::GetAchievementDefinition(FGetAchievementDefinition::Params&& Params)
 {
-	if (!Services.Get<FAuthNull>()->IsLoggedIn(Params.LocalUserId))
+	if (!Services.Get<FAuthNull>()->IsLoggedIn(Params.LocalAccountId))
 	{
 		return TOnlineResult<FGetAchievementDefinition>(Errors::InvalidUser());
 	}
@@ -132,7 +132,7 @@ TOnlineAsyncOpHandle<FQueryAchievementStates> FAchievementsNull::QueryAchievemen
 {
 	TOnlineAsyncOpRef<FQueryAchievementStates> Op = GetOp<FQueryAchievementStates>(MoveTemp(Params));
 
-	if (!Services.Get<FAuthNull>()->IsLoggedIn(Op->GetParams().LocalUserId))
+	if (!Services.Get<FAuthNull>()->IsLoggedIn(Op->GetParams().LocalAccountId))
 	{
 		Op->SetError(Errors::InvalidUser());
 		return Op->GetHandle();
@@ -145,9 +145,9 @@ TOnlineAsyncOpHandle<FQueryAchievementStates> FAchievementsNull::QueryAchievemen
 		return Op->GetHandle();
 	}
 
-	if (!AchievementStates.Contains(Op->GetParams().LocalUserId))
+	if (!AchievementStates.Contains(Op->GetParams().LocalAccountId))
 	{
-		FAchievementStateMap& LocalUserAchievementStates = AchievementStates.Emplace(Op->GetParams().LocalUserId);
+		FAchievementStateMap& LocalUserAchievementStates = AchievementStates.Emplace(Op->GetParams().LocalAccountId);
 		for (const TPair<FString, FAchievementDefinition>& Pair : *AchievementDefinitions)
 		{
 			const FString& AchievementId = Pair.Key;
@@ -162,12 +162,12 @@ TOnlineAsyncOpHandle<FQueryAchievementStates> FAchievementsNull::QueryAchievemen
 
 TOnlineResult<FGetAchievementState> FAchievementsNull::GetAchievementState(FGetAchievementState::Params&& Params) const
 {
-	if (!Services.Get<FAuthNull>()->IsLoggedIn(Params.LocalUserId))
+	if (!Services.Get<FAuthNull>()->IsLoggedIn(Params.LocalAccountId))
 	{
 		return TOnlineResult<FGetAchievementState>(Errors::InvalidUser());
 	}
 
-	const FAchievementStateMap* LocalUserAchievementStates = AchievementStates.Find(Params.LocalUserId);
+	const FAchievementStateMap* LocalUserAchievementStates = AchievementStates.Find(Params.LocalAccountId);
 	if (!LocalUserAchievementStates)
 	{
 		// Call QueryAchievementStates first
@@ -187,7 +187,7 @@ TOnlineAsyncOpHandle<FUnlockAchievements> FAchievementsNull::UnlockAchievements(
 {
 	TOnlineAsyncOpRef<FUnlockAchievements> Op = GetOp<FUnlockAchievements>(MoveTemp(Params));
 
-	if (!Services.Get<FAuthNull>()->IsLoggedIn(Op->GetParams().LocalUserId))
+	if (!Services.Get<FAuthNull>()->IsLoggedIn(Op->GetParams().LocalAccountId))
 	{
 		Op->SetError(Errors::InvalidUser());
 		return Op->GetHandle();
@@ -199,7 +199,7 @@ TOnlineAsyncOpHandle<FUnlockAchievements> FAchievementsNull::UnlockAchievements(
 		return Op->GetHandle();
 	}
 
-	FAchievementStateMap* LocalUserAchievementStates = AchievementStates.Find(Op->GetParams().LocalUserId);
+	FAchievementStateMap* LocalUserAchievementStates = AchievementStates.Find(Op->GetParams().LocalAccountId);
 	if (!LocalUserAchievementStates)
 	{
 		// Call QueryAchievementStates first
@@ -235,7 +235,7 @@ TOnlineAsyncOpHandle<FUnlockAchievements> FAchievementsNull::UnlockAchievements(
 	}
 
 	FAchievementStateUpdated AchievementStateUpdated;
-	AchievementStateUpdated.LocalUserId = Op->GetParams().LocalUserId;
+	AchievementStateUpdated.LocalAccountId = Op->GetParams().LocalAccountId;
 	AchievementStateUpdated.AchievementIds = Op->GetParams().AchievementIds;
 	OnAchievementStateUpdatedEvent.Broadcast(AchievementStateUpdated);
 
@@ -245,12 +245,12 @@ TOnlineAsyncOpHandle<FUnlockAchievements> FAchievementsNull::UnlockAchievements(
 
 TOnlineResult<FDisplayAchievementUI> FAchievementsNull::DisplayAchievementUI(FDisplayAchievementUI::Params&& Params)
 {
-	if (!Services.Get<FAuthNull>()->IsLoggedIn(Params.LocalUserId))
+	if (!Services.Get<FAuthNull>()->IsLoggedIn(Params.LocalAccountId))
 	{
 		return TOnlineResult<FDisplayAchievementUI>(Errors::InvalidUser());
 	}
 
-	const FAchievementStateMap* LocalUserAchievementStates = AchievementStates.Find(Params.LocalUserId);
+	const FAchievementStateMap* LocalUserAchievementStates = AchievementStates.Find(Params.LocalAccountId);
 	if (!LocalUserAchievementStates)
 	{
 		// Call QueryAchievementStates first
@@ -266,7 +266,7 @@ TOnlineResult<FDisplayAchievementUI> FAchievementsNull::DisplayAchievementUI(FDi
 	}
 
 	const FAchievementState& AchievementState = LocalUserAchievementStates->FindChecked(Params.AchievementId);
-	UE_LOG(LogTemp, Display, TEXT("AchievementsNull: DisplayAchievementUI LocalUserId=[%s]"), *ToLogString(Params.LocalUserId));
+	UE_LOG(LogTemp, Display, TEXT("AchievementsNull: DisplayAchievementUI LocalAccountId=[%s]"), *ToLogString(Params.LocalAccountId));
 	UE_LOG(LogTemp, Display, TEXT("AchievementsNull: DisplayAchievementUI AchievementDefinition=[%s]"), *ToLogString(*AchievementDefinition));
 	UE_LOG(LogTemp, Display, TEXT("AchievementsNull: DisplayAchievementUI AchievementState=[%s]"), *ToLogString(AchievementState));
 
diff --git a/Engine/Plugins/Online/Experimental/OnlineServicesNull/Source/Private/Online/AuthNull.cpp b/Engine/Plugins/Online/Experimental/OnlineServicesNull/Source/Private/Online/AuthNull.cpp
--- a/Engine/Plugins/Online/Experimental/OnlineServicesNull/Source/Private/Online/AuthNull.cpp
+++ b/Engine/Plugins/Online/Experimental/OnlineServicesNull/Source/Private/Online/AuthNull.cpp
@@ -90,9 +90,9 @@ TSharedPtr<FAccountInfoNull> FAccountInfoRegistryNULL::Find(FPlatformUserId Plat
 	return StaticCastSharedPtr<FAccountInfoNull>(Super::Find(PlatformUserId));
 }
 
-TSharedPtr<FAccountInfoNull> FAccountInfoRegistryNULL::Find(FAccountId AccountIdHandle) const
+TSharedPtr<FAccountInfoNull> FAccountInfoRegistryNULL::Find(FAccountId AccountId) const
 {
-	return StaticCastSharedPtr<FAccountInfoNull>(Super::Find(AccountIdHandle));
+	return StaticCastSharedPtr<FAccountInfoNull>(Super::Find(AccountId));
 }
 
 void FAccountInfoRegistryNULL::Register(const TSharedRef<FAccountInfoNull>& AccountInfoNULL)
@@ -177,12 +177,12 @@ FOnlineAccountIdRegistryNull& FOnlineAccountIdRegistryNull::Get()
 	return Instance;
 }
 
-FAccountId FOnlineAccountIdRegistryNull::Find(FString UserId) const
+FAccountId FOnlineAccountIdRegistryNull::Find(FString AccountId) const
 {
-	const FOnlineAccountIdString* Entry = StringToId.Find(UserId);
+	const FOnlineAccountIdString* Entry = StringToId.Find(AccountId);
 	if(Entry)
 	{
-		return Entry->Handle;
+		return Entry->AccountId;
 	}
 	return FAccountId();
 }
@@ -192,7 +192,7 @@ FAccountId FOnlineAccountIdRegistryNull::Find(FPlatformUserId UserId) const
 	const FOnlineAccountIdString* Entry = LocalUserMap.Find(UserId);
 	if (Entry)
 	{
-		return Entry->Handle;
+		return Entry->AccountId;
 	}
 	return FAccountId();
 }
@@ -202,27 +202,27 @@ FAccountId FOnlineAccountIdRegistryNull::Find(int32 UserIndex) const
 	const FOnlineAccountIdString* Entry = LocalUserMap.Find(FPlatformMisc::GetPlatformUserForUserIndex(UserIndex));
 	if (Entry)
 	{
-		return Entry->Handle;
+		return Entry->AccountId;
 	}
 	return FAccountId();
 }
 
-const FOnlineAccountIdString* FOnlineAccountIdRegistryNull::GetInternal(const FAccountId& Handle) const
+const FOnlineAccountIdString* FOnlineAccountIdRegistryNull::GetInternal(const FAccountId& AccountId) const
 {
-	if(Handle.IsValid() && Handle.GetOnlineServicesType() == EOnlineServices::Null && Handle.GetHandle() <= (uint32)Ids.Num())
+	if(AccountId.IsValid() && AccountId.GetOnlineServicesType() == EOnlineServices::Null && AccountId.GetHandle() <= (uint32)Ids.Num())
 	{
-		return &Ids[Handle.GetHandle()-1];
+		return &Ids[AccountId.GetHandle()-1];
 	}
 	return nullptr;
 }
 
-FAccountId FOnlineAccountIdRegistryNull::Create(FString UserId, FPlatformUserId PlatformUserId/* = PLATFORMUSERID_NONE*/)
+FAccountId FOnlineAccountIdRegistryNull::Create(FString AccountId, FPlatformUserId PlatformUserId/* = PLATFORMUSERID_NONE*/)
 {
-	FAccountId ExistingHandle = Find(UserId);
-	if(ExistingHandle.IsValid())
+	FAccountId ExistingAccountId = Find(AccountId);
+	if(ExistingAccountId.IsValid())
 	{
 		UE_LOG(LogOnlineServices, Error, TEXT("[FOnlineAccountIdRegistryNull::Create] Found a duplicate ID for local user %d."), FPlatformMisc::GetUserIndexForPlatformUser(PlatformUserId));
-		return ExistingHandle;
+		return ExistingAccountId;
 	}
 
 	if (!PlatformUserId.IsValid())
@@ -233,18 +233,18 @@ FAccountId FOnlineAccountIdRegistryNull::Create(FString UserId, FPlatformUserId
 
 	FOnlineAccountIdString& Id = Ids.Emplace_GetRef();
 	Id.AccountIndex = Ids.Num();
-	Id.Data = UserId;
-	Id.Handle = FAccountId(EOnlineServices::Null, Id.AccountIndex);
+	Id.Data = AccountId;
+	Id.AccountId = FAccountId(EOnlineServices::Null, Id.AccountIndex);
 	
-	StringToId.Add(UserId, Id);
+	StringToId.Add(AccountId, Id);
 	LocalUserMap.Add(PlatformUserId, Id);
 
-	return Id.Handle;
+	return Id.AccountId;
 }
 
-FString FOnlineAccountIdRegistryNull::ToLogString(const FAccountId& Handle) const
+FString FOnlineAccountIdRegistryNull::ToLogString(const FAccountId& AccountId) const
 {
-	if(const FOnlineAccountIdString* Id = GetInternal(Handle))
+	if(const FOnlineAccountIdString* Id = GetInternal(AccountId))
 	{
 		return Id->Data;
 	}
@@ -253,9 +253,9 @@ FString FOnlineAccountIdRegistryNull::ToLogString(const FAccountId& Handle) cons
 }
 
 
-TArray<uint8> FOnlineAccountIdRegistryNull::ToReplicationData(const FAccountId& Handle) const
+TArray<uint8> FOnlineAccountIdRegistryNull::ToReplicationData(const FAccountId& AccountId) const
 {
-	if (const FOnlineAccountIdString* Id = GetInternal(Handle))
+	if (const FOnlineAccountIdString* Id = GetInternal(AccountId))
 	{
 		TArray<uint8> ReplicationData;
 		ReplicationData.SetNumUninitialized(Id->Data.Len());
diff --git a/Engine/Plugins/Online/Experimental/OnlineServicesNull/Source/Private/Online/LeaderboardsNull.cpp b/Engine/Plugins/Online/Experimental/OnlineServicesNull/Source/Private/Online/LeaderboardsNull.cpp
--- a/Engine/Plugins/Online/Experimental/OnlineServicesNull/Source/Private/Online/LeaderboardsNull.cpp
+++ b/Engine/Plugins/Online/Experimental/OnlineServicesNull/Source/Private/Online/LeaderboardsNull.cpp
@@ -57,7 +57,7 @@ TOnlineAsyncOpHandle<FWriteLeaderboardScores> FLeaderboardsNull::WriteLeaderboar
 {
 	TOnlineAsyncOpRef<FWriteLeaderboardScores> Op = GetOp<FWriteLeaderboardScores>(MoveTemp(Params));
 
-	if (!Services.Get<FAuthNull>()->IsLoggedIn(Op->GetParams().LocalUserId))
+	if (!Services.Get<FAuthNull>()->IsLoggedIn(Op->GetParams().LocalAccountId))
 	{
 		Op->SetError(Errors::InvalidUser());
 		return Op->GetHandle();
@@ -88,7 +88,7 @@ TOnlineAsyncOpHandle<FWriteLeaderboardScores> FLeaderboardsNull::WriteLeaderboar
 		while (CurrentNode)
 		{
 			const FUserScoreNull& UserScore = CurrentNode->GetValue();
-			if (UserScore.UserId == InAsyncOp.GetParams().LocalUserId)
+			if (UserScore.AccountId == InAsyncOp.GetParams().LocalAccountId)
 			{
 				UpdatedScore = Private::GetUpdatedScoreByLeaderboardDefinition(LeaderboardDefinition, UserScore.Score, InAsyncOp.GetParams().Score);
 				UserScoreList.RemoveNode(CurrentNode);
@@ -101,7 +101,7 @@ TOnlineAsyncOpHandle<FWriteLeaderboardScores> FLeaderboardsNull::WriteLeaderboar
 		// Insert into list by order
 		FUserScoreNull UserScoreToInsert;
 		UserScoreToInsert.Score = UpdatedScore;
-		UserScoreToInsert.UserId = InAsyncOp.GetParams().LocalUserId;
+		UserScoreToInsert.AccountId = InAsyncOp.GetParams().LocalAccountId;
 
 		bool bInserted = false;
 		CurrentNode = UserScoreList.GetHead();
@@ -147,7 +147,7 @@ TOnlineAsyncOpHandle<FReadEntriesForUsers> FLeaderboardsNull::ReadEntriesForUser
 {
 	TOnlineAsyncOpRef<FReadEntriesForUsers> Op = GetOp<FReadEntriesForUsers>(MoveTemp(Params));
 
-	if (!Services.Get<FAuthNull>()->IsLoggedIn(Op->GetParams().LocalUserId))
+	if (!Services.Get<FAuthNull>()->IsLoggedIn(Op->GetParams().LocalAccountId))
 	{
 		Op->SetError(Errors::InvalidUser());
 		return Op->GetHandle();
@@ -159,18 +159,18 @@ TOnlineAsyncOpHandle<FReadEntriesForUsers> FLeaderboardsNull::ReadEntriesForUser
 
 		if (FLeaderboardDataNull* LeaderboardData = LeaderboardsData.FindByPredicate(FFindLeaderboardDataByName(InAsyncOp.GetParams().BoardName)))
 		{
-			for (const FAccountId& UserId : InAsyncOp.GetParams().UserIds)
+			for (const FAccountId& AccountId : InAsyncOp.GetParams().AccountIds)
 			{
 				uint32 Index = 0;
 				TDoubleLinkedList<FUserScoreNull>::TDoubleLinkedListNode* CurrentNode = LeaderboardData->UserScoreList.GetHead();
 				while (CurrentNode)
 				{
 					const FUserScoreNull& UserScore = CurrentNode->GetValue();
-					// When needed, use map/set instead of array for UserIds in Params, to improve performance
-					if (UserScore.UserId == UserId)
+					// When needed, use map/set instead of array for AccountIds in Params, to improve performance
+					if (UserScore.AccountId == AccountId)
 					{
 						FLeaderboardEntry& LeaderboardEntry = Result.Entries.Emplace_GetRef();
-						LeaderboardEntry.UserId = UserId;
+						LeaderboardEntry.AccountId = AccountId;
 						LeaderboardEntry.Rank = Index;
 						LeaderboardEntry.Score = UserScore.Score;
 						break;
@@ -204,7 +204,7 @@ void ReadEntriesInRange(TDoubleLinkedList<FUserScoreNull>& UserScoreList, uint32
 		{
 			const FUserScoreNull& UserScore = CurrentNode->GetValue();
 			FLeaderboardEntry& LeaderboardEntry = OutEntries.Emplace_GetRef();
-			LeaderboardEntry.UserId = UserScore.UserId;
+			LeaderboardEntry.AccountId = UserScore.AccountId;
 			LeaderboardEntry.Rank = Index;
 			LeaderboardEntry.Score = UserScore.Score;
 		}
@@ -220,7 +220,7 @@ TOnlineAsyncOpHandle<FReadEntriesAroundRank> FLeaderboardsNull::ReadEntriesAroun
 {
 	TOnlineAsyncOpRef<FReadEntriesAroundRank> Op = GetOp<FReadEntriesAroundRank>(MoveTemp(Params));
 
-	if (!Services.Get<FAuthNull>()->IsLoggedIn(Op->GetParams().LocalUserId))
+	if (!Services.Get<FAuthNull>()->IsLoggedIn(Op->GetParams().LocalAccountId))
 	{
 		Op->SetError(Errors::InvalidUser());
 		return Op->GetHandle();
@@ -252,7 +252,7 @@ TOnlineAsyncOpHandle<FReadEntriesAroundUser> FLeaderboardsNull::ReadEntriesAroun
 {
 	TOnlineAsyncOpRef<FReadEntriesAroundUser> Op = GetOp<FReadEntriesAroundUser>(MoveTemp(Params));
 
-	if (!Services.Get<FAuthNull>()->IsLoggedIn(Op->GetParams().LocalUserId))
+	if (!Services.Get<FAuthNull>()->IsLoggedIn(Op->GetParams().LocalAccountId))
 	{
 		Op->SetError(Errors::InvalidUser());
 		return Op->GetHandle();
@@ -276,7 +276,7 @@ TOnlineAsyncOpHandle<FReadEntriesAroundUser> FLeaderboardsNull::ReadEntriesAroun
 			while (CurrentNode)
 			{
 				const FUserScoreNull& UserScore = CurrentNode->GetValue();
-				if (UserScore.UserId == InAsyncOp.GetParams().UserId)
+				if (UserScore.AccountId == InAsyncOp.GetParams().AccountId)
 				{
 					FoundUser = true;
 					break;
diff --git a/Engine/Plugins/Online/Experimental/OnlineServicesNull/Source/Private/Online/LobbiesNull.cpp b/Engine/Plugins/Online/Experimental/OnlineServicesNull/Source/Private/Online/LobbiesNull.cpp
--- a/Engine/Plugins/Online/Experimental/OnlineServicesNull/Source/Private/Online/LobbiesNull.cpp
+++ b/Engine/Plugins/Online/Experimental/OnlineServicesNull/Source/Private/Online/LobbiesNull.cpp
@@ -147,7 +147,7 @@ TSharedRef<FLobbyNull> FLobbiesNull::CreateNamedLobby(const FCreateLobby::Params
 	//CHECK(!NamedLobbies.Contains(Params.LocalName))
 	TSharedRef<FLobbyNull> Lobby = NamedLobbies.Add(Params.LocalName, MakeShared<FLobbyNull>());
 	Lobby->Data->LobbyId = FOnlineLobbyIdRegistryNull::Get().GetNext();
-	Lobby->Data->OwnerAccountId = Params.LocalUserId;
+	Lobby->Data->OwnerAccountId = Params.LocalAccountId;
 	Lobby->Data->LocalName = Params.LocalName;
 	//Lobby.Data.SchemaName = ???
 	Lobby->Data->JoinPolicy = ELobbyJoinPolicy::PublicAdvertised; // temp
@@ -159,12 +159,12 @@ TSharedRef<FLobbyNull> FLobbiesNull::CreateNamedLobby(const FCreateLobby::Params
 	{
 		TSharedRef<FLobbyMember> NewMember = MakeShared<FLobbyMember>();
 
-		NewMember->AccountId = Data.LocalUserId;
-		NewMember->PlatformAccountId = Data.LocalUserId;
+		NewMember->AccountId = Data.LocalAccountId;
+		NewMember->PlatformAccountId = Data.LocalAccountId;
 		NewMember->PlatformDisplayName = TEXT("TEMP");
 		NewMember->Attributes = Data.Attributes;
 
-		Lobby->Data->Members.Add(Data.LocalUserId, NewMember);
+		Lobby->Data->Members.Add(Data.LocalAccountId, NewMember);
 
 	}
 
@@ -299,10 +299,10 @@ void FLobbiesNull::ReadLobbyFromPacket(FNboSerializeFromBuffer& Packet, const TS
 	UE_LOG(LogTemp, Verbose, TEXT("Reading session information from server"));
 
 	/** Owner of the session */
-	FAccountId OwningUserId;
+	FAccountId OwningAccountId;
 	FString OwningUserName;
 	uint32 NumOpenPrivateConnections, NumOpenPublicConnections;
-	SerializeFromBuffer(Packet, OwningUserId);
+	SerializeFromBuffer(Packet, OwningAccountId);
 	SerializeFromBuffer(Packet, Session->Data->Attributes);
 	Packet	>> OwningUserName
 		>> Session->HostAddrIp
@@ -310,7 +310,7 @@ void FLobbiesNull::ReadLobbyFromPacket(FNboSerializeFromBuffer& Packet, const TS
 		>> NumOpenPrivateConnections
 		>> NumOpenPublicConnections;
 
-	Session->Data->OwnerAccountId = OwningUserId;
+	Session->Data->OwnerAccountId = OwningAccountId;
 	Session->HostAddr = ISocketSubsystem::Get(PLATFORM_SOCKETSUBSYSTEM)->CreateInternetAddr();
 	Session->HostAddr->SetIp(Session->HostAddrIp);
 	Session->HostAddr->SetPort(Session->HostAddrPort);
@@ -320,10 +320,10 @@ void FLobbiesNull::ReadLobbyFromPacket(FNboSerializeFromBuffer& Packet, const TS
 
 	// todo: make all the members (not just the host
 	TSharedRef<FLobbyMember> NewMember = MakeShared<FLobbyMember>();
-	NewMember->AccountId = OwningUserId;
-	NewMember->PlatformAccountId = OwningUserId;
+	NewMember->AccountId = OwningAccountId;
+	NewMember->PlatformAccountId = OwningAccountId;
 	NewMember->PlatformDisplayName = TEXT("TEMP");
-	Session->Data->Members.Add(OwningUserId, NewMember);
+	Session->Data->Members.Add(OwningAccountId, NewMember);
 	AllLobbies.Add(Session->Data->LobbyId, Session);
 
 	// todo: host info, per game settings
diff --git a/Engine/Plugins/Online/Experimental/OnlineServicesNull/Source/Private/Online/OnlineServicesNull.cpp b/Engine/Plugins/Online/Experimental/OnlineServicesNull/Source/Private/Online/OnlineServicesNull.cpp
--- a/Engine/Plugins/Online/Experimental/OnlineServicesNull/Source/Private/Online/OnlineServicesNull.cpp
+++ b/Engine/Plugins/Online/Experimental/OnlineServicesNull/Source/Private/Online/OnlineServicesNull.cpp
@@ -56,7 +56,7 @@ TOnlineResult<FGetResolvedConnectString> FOnlineServicesNull::GetResolvedConnect
 		ILobbiesPtr LobbiesPtr = GetLobbiesInterface();
 		check(LobbiesPtr);
 
-		TOnlineResult<FGetJoinedLobbies> JoinedLobbies = LobbiesPtr->GetJoinedLobbies({ Params.LocalUserId });
+		TOnlineResult<FGetJoinedLobbies> JoinedLobbies = LobbiesPtr->GetJoinedLobbies({ Params.LocalAccountId });
 		if (JoinedLobbies.IsOk())
 		{
 			for (TSharedRef<const FLobby>& Lobby : JoinedLobbies.GetOkValue().Lobbies)
diff --git a/Engine/Plugins/Online/Experimental/OnlineServicesNull/Source/Private/Online/PresenceNull.cpp b/Engine/Plugins/Online/Experimental/OnlineServicesNull/Source/Private/Online/PresenceNull.cpp
--- a/Engine/Plugins/Online/Experimental/OnlineServicesNull/Source/Private/Online/PresenceNull.cpp
+++ b/Engine/Plugins/Online/Experimental/OnlineServicesNull/Source/Private/Online/PresenceNull.cpp
@@ -34,12 +34,12 @@ TOnlineAsyncOpHandle<FQueryPresence> FPresenceNull::QueryPresence(FQueryPresence
 
 		if(Params.bListenToChanges)
 		{
-			PresenceListeners.FindOrAdd(Params.LocalUserId).Add(Params.TargetUserId);
+			PresenceListeners.FindOrAdd(Params.LocalAccountId).Add(Params.TargetAccountId);
 		}
 
-		if(Presences.Contains(Params.TargetUserId))
+		if(Presences.Contains(Params.TargetAccountId))
 		{
-			InAsyncOp.SetResult({Presences.FindChecked(Params.TargetUserId)});
+			InAsyncOp.SetResult({Presences.FindChecked(Params.TargetAccountId)});
 		}
 		else
 		{
@@ -53,9 +53,9 @@ TOnlineAsyncOpHandle<FQueryPresence> FPresenceNull::QueryPresence(FQueryPresence
 
 TOnlineResult<FGetCachedPresence> FPresenceNull::GetCachedPresence(FGetCachedPresence::Params&& Params)
 {
-	if (Presences.Contains(Params.LocalUserId))
+	if (Presences.Contains(Params.LocalAccountId))
 	{
-		return TOnlineResult<FGetCachedPresence>({Presences.FindChecked(Params.LocalUserId)});
+		return TOnlineResult<FGetCachedPresence>({Presences.FindChecked(Params.LocalAccountId)});
 	}
 	else
 	{
@@ -71,11 +71,11 @@ TOnlineAsyncOpHandle<FUpdatePresence> FPresenceNull::UpdatePresence(FUpdatePrese
 	Op->Then([this](TOnlineAsyncOp<FUpdatePresence>& InAsyncOp) mutable
 	{
 		const FUpdatePresence::Params& Params = InAsyncOp.GetParams();
-		Presences.Add(Params.LocalUserId, Params.Presence);
+		Presences.Add(Params.LocalAccountId, Params.Presence);
 
 		for (const TPair<FAccountId, TSet<FAccountId>>& Pairs : PresenceListeners)
 		{
-			if (Pairs.Value.Contains(Params.LocalUserId))
+			if (Pairs.Value.Contains(Params.LocalAccountId))
 			{
 				OnPresenceUpdatedEvent.Broadcast({Pairs.Key, Params.Presence});
 			}
@@ -98,17 +98,17 @@ TOnlineAsyncOpHandle<FPartialUpdatePresence> FPresenceNull::PartialUpdatePresenc
 		TSharedRef<FUserPresence> NewPresence = MakeShared<FUserPresence>();
 		const FPartialUpdatePresence::Params::FMutations& Mutations = InAsyncOp.GetParams().Mutations;
 		
-		if (Presences.Contains(InAsyncOp.GetParams().LocalUserId))
+		if (Presences.Contains(InAsyncOp.GetParams().LocalAccountId))
 		{
-			*NewPresence = *Presences.FindChecked(InAsyncOp.GetParams().LocalUserId);
+			*NewPresence = *Presences.FindChecked(InAsyncOp.GetParams().LocalAccountId);
 		}
 
 		TSharedRef<FUserPresence> MutatedPresence = ApplyPresenceMutations(*NewPresence, Mutations);
-		Presences.Add(InAsyncOp.GetParams().LocalUserId, MutatedPresence);
+		Presences.Add(InAsyncOp.GetParams().LocalAccountId, MutatedPresence);
 
 		for (const TPair<FAccountId, TSet<FAccountId>>& Pairs : PresenceListeners)
 		{
-			if (Pairs.Value.Contains(InAsyncOp.GetParams().LocalUserId))
+			if (Pairs.Value.Contains(InAsyncOp.GetParams().LocalAccountId))
 			{
 				OnPresenceUpdatedEvent.Broadcast({ Pairs.Key, MutatedPresence });
 			}
diff --git a/Engine/Plugins/Online/Experimental/OnlineServicesNull/Source/Private/Online/SessionsNull.cpp b/Engine/Plugins/Online/Experimental/OnlineServicesNull/Source/Private/Online/SessionsNull.cpp
--- a/Engine/Plugins/Online/Experimental/OnlineServicesNull/Source/Private/Online/SessionsNull.cpp
+++ b/Engine/Plugins/Online/Experimental/OnlineServicesNull/Source/Private/Online/SessionsNull.cpp
@@ -36,7 +36,7 @@ void FSessionsNull::AppendSessionToPacket(FNboSerializeToBuffer& Packet, const F
 	using namespace NboSerializerNullSvc;
 
 	SerializeToBuffer(Packet, Session);
-	SerializeToBuffer(Packet, Session.OwnerUserId);
+	SerializeToBuffer(Packet, Session.OwnerAccountId);
 	SerializeToBuffer(Packet, Session.SessionId);
 	SerializeToBuffer(Packet, Session.SessionSettings.RegisteredPlayers);
 	SerializeToBuffer(Packet, Session.SessionSettings.SessionMembers);
@@ -48,7 +48,7 @@ void FSessionsNull::ReadSessionFromPacket(FNboSerializeFromBuffer& Packet, FSess
 	using namespace NboSerializerNullSvc;
 
 	SerializeFromBuffer(Packet, Session);
-	SerializeFromBuffer(Packet, Session.OwnerUserId);
+	SerializeFromBuffer(Packet, Session.OwnerAccountId);
 	SerializeFromBuffer(Packet, Session.SessionId);
 	SerializeFromBuffer(Packet, Session.SessionSettings.RegisteredPlayers);
 	SerializeFromBuffer(Packet, Session.SessionSettings.SessionMembers);
diff --git a/Engine/Plugins/Online/Experimental/OnlineServicesNull/Source/Private/Online/StatsNull.cpp b/Engine/Plugins/Online/Experimental/OnlineServicesNull/Source/Private/Online/StatsNull.cpp
--- a/Engine/Plugins/Online/Experimental/OnlineServicesNull/Source/Private/Online/StatsNull.cpp
+++ b/Engine/Plugins/Online/Experimental/OnlineServicesNull/Source/Private/Online/StatsNull.cpp
@@ -16,7 +16,7 @@ TOnlineAsyncOpHandle<FUpdateStats> FStatsNull::UpdateStats(FUpdateStats::Params&
 {
 	TOnlineAsyncOpRef<FUpdateStats> Op = GetOp<FUpdateStats>(MoveTemp(Params));
 
-	if (!Services.Get<FAuthNull>()->IsLoggedIn(Op->GetParams().LocalUserId))
+	if (!Services.Get<FAuthNull>()->IsLoggedIn(Op->GetParams().LocalAccountId))
 	{
 		Op->SetError(Errors::InvalidUser());
 		return Op->GetHandle();
@@ -26,11 +26,11 @@ TOnlineAsyncOpHandle<FUpdateStats> FStatsNull::UpdateStats(FUpdateStats::Params&
 	{
 		for (const FUserStats& UpdateUserStats : InAsyncOp.GetParams().UpdateUsersStats)
 		{
-			FUserStats* ExistingUserStats = UsersStats.FindByPredicate(FFindUserStatsByAccountId(UpdateUserStats.UserId));
+			FUserStats* ExistingUserStats = UsersStats.FindByPredicate(FFindUserStatsByAccountId(UpdateUserStats.AccountId));
 			if (!ExistingUserStats)
 			{
 				ExistingUserStats = &UsersStats.Emplace_GetRef();
-				ExistingUserStats->UserId = UpdateUserStats.UserId;
+				ExistingUserStats->AccountId = UpdateUserStats.AccountId;
 			}
 
 			TMap<FString, FStatValue>& UserStats = ExistingUserStats->Stats;
@@ -104,7 +104,7 @@ TOnlineAsyncOpHandle<FQueryStats> FStatsNull::QueryStats(FQueryStats::Params&& P
 {
 	TOnlineAsyncOpRef<FQueryStats> Op = GetOp<FQueryStats>(MoveTemp(Params));
 
-	if (!Services.Get<FAuthNull>()->IsLoggedIn(Op->GetParams().LocalUserId))
+	if (!Services.Get<FAuthNull>()->IsLoggedIn(Op->GetParams().LocalAccountId))
 	{
 		Op->SetError(Errors::InvalidUser());
 		return Op->GetHandle();
@@ -114,7 +114,7 @@ TOnlineAsyncOpHandle<FQueryStats> FStatsNull::QueryStats(FQueryStats::Params&& P
 	{
 		FQueryStats::Result Result;
 
-		if (FUserStats* ExistingUserStats = UsersStats.FindByPredicate(FFindUserStatsByAccountId(InAsyncOp.GetParams().TargetUserId)))
+		if (FUserStats* ExistingUserStats = UsersStats.FindByPredicate(FFindUserStatsByAccountId(InAsyncOp.GetParams().TargetAccountId)))
 		{
 			Result.Stats = ExistingUserStats->Stats;
 		}
@@ -129,7 +129,7 @@ TOnlineAsyncOpHandle<FQueryStats> FStatsNull::QueryStats(FQueryStats::Params&& P
 TOnlineAsyncOpHandle<FBatchQueryStats> FStatsNull::BatchQueryStats(FBatchQueryStats::Params&& Params)
 {
 	TOnlineAsyncOpRef<FBatchQueryStats> Op = GetOp<FBatchQueryStats>(MoveTemp(Params));
-	if (!Services.Get<FAuthNull>()->IsLoggedIn(Op->GetParams().LocalUserId))
+	if (!Services.Get<FAuthNull>()->IsLoggedIn(Op->GetParams().LocalAccountId))
 	{
 		Op->SetError(Errors::InvalidUser());
 		return Op->GetHandle();
@@ -139,9 +139,9 @@ TOnlineAsyncOpHandle<FBatchQueryStats> FStatsNull::BatchQueryStats(FBatchQuerySt
 	{
 		FBatchQueryStats::Result Result;
 
-		for (const FAccountId& TargetUserId : InAsyncOp.GetParams().TargetUserIds)
+		for (const FAccountId& TargetAccountId : InAsyncOp.GetParams().TargetAccountIds)
 		{
-			if (FUserStats* ExistingUserStats = UsersStats.FindByPredicate(FFindUserStatsByAccountId(TargetUserId)))
+			if (FUserStats* ExistingUserStats = UsersStats.FindByPredicate(FFindUserStatsByAccountId(TargetAccountId)))
 			{
 				Result.UsersStats.Emplace(*ExistingUserStats);
 			}
@@ -159,7 +159,7 @@ TOnlineAsyncOpHandle<FResetStats> FStatsNull::ResetStats(FResetStats::Params&& P
 {
 	TOnlineAsyncOpRef<FResetStats> Op = GetOp<FResetStats>(MoveTemp(Params));
 
-	if (!Services.Get<FAuthNull>()->IsLoggedIn(Op->GetParams().LocalUserId))
+	if (!Services.Get<FAuthNull>()->IsLoggedIn(Op->GetParams().LocalAccountId))
 	{
 		Op->SetError(Errors::InvalidUser());
 		return Op->GetHandle();
@@ -167,13 +167,13 @@ TOnlineAsyncOpHandle<FResetStats> FStatsNull::ResetStats(FResetStats::Params&& P
 
 	Op->Then([this](TOnlineAsyncOp<FResetStats>& InAsyncOp)
 	{
-		uint32 Index = UsersStats.IndexOfByPredicate(FFindUserStatsByAccountId(InAsyncOp.GetParams().LocalUserId));
+		uint32 Index = UsersStats.IndexOfByPredicate(FFindUserStatsByAccountId(InAsyncOp.GetParams().LocalAccountId));
 		if (Index != INDEX_NONE)
 		{
 			UsersStats.RemoveAt(Index);
 		}
 
-		FFindUserStatsByAccountId FindUserStatsByAccountId(InAsyncOp.GetParams().LocalUserId);
+		FFindUserStatsByAccountId FindUserStatsByAccountId(InAsyncOp.GetParams().LocalAccountId);
 
 		UsersStats.RemoveAll(FindUserStatsByAccountId);
 		CachedUsersStats.RemoveAll(FindUserStatsByAccountId);
diff --git a/Engine/Plugins/Online/Experimental/OnlineServicesNull/Source/Private/Online/TitleFileNull.cpp b/Engine/Plugins/Online/Experimental/OnlineServicesNull/Source/Private/Online/TitleFileNull.cpp
--- a/Engine/Plugins/Online/Experimental/OnlineServicesNull/Source/Private/Online/TitleFileNull.cpp
+++ b/Engine/Plugins/Online/Experimental/OnlineServicesNull/Source/Private/Online/TitleFileNull.cpp
@@ -45,7 +45,7 @@ TOnlineAsyncOpHandle<FTitleFileEnumerateFiles> FTitleFileNull::EnumerateFiles(FT
 {
 	TOnlineAsyncOpRef<FTitleFileEnumerateFiles> Op = GetOp<FTitleFileEnumerateFiles>(MoveTemp(Params));
 
-	if (!Services.Get<FAuthNull>()->IsLoggedIn(Op->GetParams().LocalUserId))
+	if (!Services.Get<FAuthNull>()->IsLoggedIn(Op->GetParams().LocalAccountId))
 	{
 		Op->SetError(Errors::InvalidUser());
 		return Op->GetHandle();
@@ -59,7 +59,7 @@ TOnlineAsyncOpHandle<FTitleFileEnumerateFiles> FTitleFileNull::EnumerateFiles(FT
 
 TOnlineResult<FTitleFileGetEnumeratedFiles> FTitleFileNull::GetEnumeratedFiles(FTitleFileGetEnumeratedFiles::Params&& Params)
 {
-	if (!Services.Get<FAuthNull>()->IsLoggedIn(Params.LocalUserId))
+	if (!Services.Get<FAuthNull>()->IsLoggedIn(Params.LocalAccountId))
 	{
 		return TOnlineResult<FTitleFileGetEnumeratedFiles>(Errors::InvalidUser());
 	}
@@ -80,7 +80,7 @@ TOnlineAsyncOpHandle<FTitleFileReadFile> FTitleFileNull::ReadFile(FTitleFileRead
 {
 	TOnlineAsyncOpRef<FTitleFileReadFile> Op = GetOp<FTitleFileReadFile>(MoveTemp(Params));
 
-	if (!Services.Get<FAuthNull>()->IsLoggedIn(Op->GetParams().LocalUserId))
+	if (!Services.Get<FAuthNull>()->IsLoggedIn(Op->GetParams().LocalAccountId))
 	{
 		Op->SetError(Errors::InvalidUser());
 		return Op->GetHandle();
diff --git a/Engine/Plugins/Online/Experimental/OnlineServicesNull/Source/Private/Online/UserFileNull.cpp b/Engine/Plugins/Online/Experimental/OnlineServicesNull/Source/Private/Online/UserFileNull.cpp
--- a/Engine/Plugins/Online/Experimental/OnlineServicesNull/Source/Private/Online/UserFileNull.cpp
+++ b/Engine/Plugins/Online/Experimental/OnlineServicesNull/Source/Private/Online/UserFileNull.cpp
@@ -48,13 +48,13 @@ TOnlineAsyncOpHandle<FUserFileEnumerateFiles> FUserFileNull::EnumerateFiles(FUse
 	TOnlineAsyncOpRef<FUserFileEnumerateFiles> Op = GetOp<FUserFileEnumerateFiles>(MoveTemp(InParams));
 	const FUserFileEnumerateFiles::Params& Params = Op->GetParams();
 
-	if (!Services.Get<FAuthNull>()->IsLoggedIn(Params.LocalUserId))
+	if (!Services.Get<FAuthNull>()->IsLoggedIn(Params.LocalAccountId))
 	{
 		Op->SetError(Errors::InvalidUser());
 		return Op->GetHandle();
 	}
 
-	FUserState& UserState = GetUserState(Params.LocalUserId);
+	FUserState& UserState = GetUserState(Params.LocalAccountId);
 	UserState.bEnumerated = true;
 
 	Op->SetResult({});
@@ -63,12 +63,12 @@ TOnlineAsyncOpHandle<FUserFileEnumerateFiles> FUserFileNull::EnumerateFiles(FUse
 
 TOnlineResult<FUserFileGetEnumeratedFiles> FUserFileNull::GetEnumeratedFiles(FUserFileGetEnumeratedFiles::Params&& Params)
 {
-	if (!Services.Get<FAuthNull>()->IsLoggedIn(Params.LocalUserId))
+	if (!Services.Get<FAuthNull>()->IsLoggedIn(Params.LocalAccountId))
 	{
 		return TOnlineResult<FUserFileGetEnumeratedFiles>(Errors::InvalidUser());
 	}
 
-	FUserState& FileState = GetUserState(Params.LocalUserId);
+	FUserState& FileState = GetUserState(Params.LocalAccountId);
 	if (!FileState.bEnumerated)
 	{
 		// Need to call EnumerateFiles first.
@@ -85,7 +85,7 @@ TOnlineAsyncOpHandle<FUserFileReadFile> FUserFileNull::ReadFile(FUserFileReadFil
 	TOnlineAsyncOpRef<FUserFileReadFile> Op = GetOp<FUserFileReadFile>(MoveTemp(InParams));
 	const FUserFileReadFile::Params& Params = Op->GetParams();
 
-	if (!Services.Get<FAuthNull>()->IsLoggedIn(Params.LocalUserId))
+	if (!Services.Get<FAuthNull>()->IsLoggedIn(Params.LocalAccountId))
 	{
 		Op->SetError(Errors::InvalidUser());
 		return Op->GetHandle();
@@ -97,7 +97,7 @@ TOnlineAsyncOpHandle<FUserFileReadFile> FUserFileNull::ReadFile(FUserFileReadFil
 		return Op->GetHandle();
 	}
 
-	FUserState& UserState = GetUserState(Params.LocalUserId);
+	FUserState& UserState = GetUserState(Params.LocalAccountId);
 	const FUserFileContentsRef* Found = UserState.Files.Find(Params.Filename);
 	if (!Found)
 	{
@@ -114,7 +114,7 @@ TOnlineAsyncOpHandle<FUserFileWriteFile> FUserFileNull::WriteFile(FUserFileWrite
 	TOnlineAsyncOpRef<FUserFileWriteFile> Op = GetOp<FUserFileWriteFile>(MoveTemp(InParams));
 	const FUserFileWriteFile::Params& Params = Op->GetParams();
 
-	if (!Services.Get<FAuthNull>()->IsLoggedIn(Params.LocalUserId))
+	if (!Services.Get<FAuthNull>()->IsLoggedIn(Params.LocalAccountId))
 	{
 		Op->SetError(Errors::InvalidUser());
 		return Op->GetHandle();
@@ -126,7 +126,7 @@ TOnlineAsyncOpHandle<FUserFileWriteFile> FUserFileNull::WriteFile(FUserFileWrite
 		return Op->GetHandle();
 	}
 
-	FUserState& UserState = GetUserState(Params.LocalUserId);
+	FUserState& UserState = GetUserState(Params.LocalAccountId);
 	UserState.Files.Emplace(Params.Filename, MakeShared<FUserFileContents>(Params.FileContents));
 	Op->SetResult({});
 	return Op->GetHandle();
@@ -137,7 +137,7 @@ TOnlineAsyncOpHandle<FUserFileCopyFile> FUserFileNull::CopyFile(FUserFileCopyFil
 	TOnlineAsyncOpRef<FUserFileCopyFile> Op = GetOp<FUserFileCopyFile>(MoveTemp(InParams));
 	const FUserFileCopyFile::Params& Params = Op->GetParams();
 
-	if (!Services.Get<FAuthNull>()->IsLoggedIn(Params.LocalUserId))
+	if (!Services.Get<FAuthNull>()->IsLoggedIn(Params.LocalAccountId))
 	{
 		Op->SetError(Errors::InvalidUser());
 		return Op->GetHandle();
@@ -149,7 +149,7 @@ TOnlineAsyncOpHandle<FUserFileCopyFile> FUserFileNull::CopyFile(FUserFileCopyFil
 		return Op->GetHandle();
 	}
 
-	FUserState& UserState = GetUserState(Params.LocalUserId);
+	FUserState& UserState = GetUserState(Params.LocalAccountId);
 	const FUserFileContentsRef* SourceFile = UserState.Files.Find(Params.SourceFilename);
 	if (!SourceFile)
 	{
@@ -167,7 +167,7 @@ TOnlineAsyncOpHandle<FUserFileDeleteFile> FUserFileNull::DeleteFile(FUserFileDel
 	TOnlineAsyncOpRef<FUserFileDeleteFile> Op = GetOp<FUserFileDeleteFile>(MoveTemp(InParams));
 	const FUserFileDeleteFile::Params& Params = Op->GetParams();
 
-	if (!Services.Get<FAuthNull>()->IsLoggedIn(Params.LocalUserId))
+	if (!Services.Get<FAuthNull>()->IsLoggedIn(Params.LocalAccountId))
 	{
 		Op->SetError(Errors::InvalidUser());
 		return Op->GetHandle();
@@ -179,7 +179,7 @@ TOnlineAsyncOpHandle<FUserFileDeleteFile> FUserFileNull::DeleteFile(FUserFileDel
 		return Op->GetHandle();
 	}
 
-	FUserState& UserState = GetUserState(Params.LocalUserId);
+	FUserState& UserState = GetUserState(Params.LocalAccountId);
 	const bool bRemoved = UserState.Files.Remove(Params.Filename) != 0;
 	if (!bRemoved)
 	{
diff --git a/Engine/Plugins/Online/Experimental/OnlineServicesNull/Source/Public/Online/AuthNull.h b/Engine/Plugins/Online/Experimental/OnlineServicesNull/Source/Public/Online/AuthNull.h
--- a/Engine/Plugins/Online/Experimental/OnlineServicesNull/Source/Public/Online/AuthNull.h
+++ b/Engine/Plugins/Online/Experimental/OnlineServicesNull/Source/Public/Online/AuthNull.h
@@ -14,7 +14,7 @@ class FOnlineAccountIdString
 public:
 	FString Data;
 	int32 AccountIndex;
-	FAccountId Handle;
+	FAccountId AccountId;
 };
 
 class FOnlineAccountIdRegistryNull : public IOnlineAccountIdRegistry
@@ -22,22 +22,22 @@ class FOnlineAccountIdRegistryNull : public IOnlineAccountIdRegistry
 public:
 	static FOnlineAccountIdRegistryNull& Get();
 
-	FAccountId Find(FString UserId) const;
+	FAccountId Find(FString AccountId) const;
 	FAccountId Find(FPlatformUserId UserId) const;
 	FAccountId Find(int32 UserIndex) const;
 
-	FAccountId Create(FString UserId, FPlatformUserId PlatformUserId = PLATFORMUSERID_NONE);
+	FAccountId Create(FString AccountId, FPlatformUserId PlatformUserId = PLATFORMUSERID_NONE);
 
 	// Begin IOnlineAccountIdRegistry
-	virtual FString ToLogString(const FAccountId& Handle) const override;
-	virtual TArray<uint8> ToReplicationData(const FAccountId& Handle) const override;
+	virtual FString ToLogString(const FAccountId& AccountId) const override;
+	virtual TArray<uint8> ToReplicationData(const FAccountId& AccountId) const override;
 	virtual FAccountId FromReplicationData(const TArray<uint8>& ReplicationString) override;
 	// End IOnlineAccountIdRegistry
 
 	virtual ~FOnlineAccountIdRegistryNull() = default;
 
 private:
-	const FOnlineAccountIdString* GetInternal(const FAccountId& Handle) const;
+	const FOnlineAccountIdString* GetInternal(const FAccountId& AccountId) const;
 
 	// how much of these are actually necessary?
 	TArray<FOnlineAccountIdString> Ids;
diff --git a/Engine/Plugins/Online/Experimental/OnlineServicesNull/Source/Public/Online/LeaderboardsNull.h b/Engine/Plugins/Online/Experimental/OnlineServicesNull/Source/Public/Online/LeaderboardsNull.h
--- a/Engine/Plugins/Online/Experimental/OnlineServicesNull/Source/Public/Online/LeaderboardsNull.h
+++ b/Engine/Plugins/Online/Experimental/OnlineServicesNull/Source/Public/Online/LeaderboardsNull.h
@@ -12,7 +12,7 @@ class FOnlineServicesNull;
 
 struct FUserScoreNull
 {
-	FAccountId UserId;
+	FAccountId AccountId;
 	uint64 Score;
 };
 
