OSSv2:
- Rename type FOnlineLobbyIdHandle to FLobbyId;
- Rename variable LobbyIdHandle/Handle to LobbyId.

diff --git a/Engine/Plugins/Online/Experimental/OnlineServices/Source/OnlineServicesCommon/Private/Online/LobbiesCommonTests.cpp b/Engine/Plugins/Online/Experimental/OnlineServices/Source/OnlineServicesCommon/Private/Online/LobbiesCommonTests.cpp
--- a/Engine/Plugins/Online/Experimental/OnlineServices/Source/OnlineServicesCommon/Private/Online/LobbiesCommonTests.cpp
+++ b/Engine/Plugins/Online/Experimental/OnlineServices/Source/OnlineServicesCommon/Private/Online/LobbiesCommonTests.cpp
@@ -114,7 +114,7 @@ bool FClientLobbyDataTest::RunTest(const FString& Parameters)
 	{
 		FLobbyEvents LobbyEvents;
 		FLobbyEventCapture EventCapture(LobbyEvents);
-		FOnlineLobbyIdHandle LobbyId(EOnlineServices::Null, 1);
+		FLobbyId LobbyId(EOnlineServices::Null, 1);
 		FClientLobbyData ClientData(LobbyId);
 
 		FAccountId User1(EOnlineServices::Null, 1);
@@ -267,7 +267,7 @@ bool FClientLobbyDataTest::RunTest(const FString& Parameters)
 	{
 		FLobbyEvents LobbyEvents;
 		FLobbyEventCapture EventCapture(LobbyEvents);
-		FOnlineLobbyIdHandle LobbyId(EOnlineServices::Null, 1);
+		FLobbyId LobbyId(EOnlineServices::Null, 1);
 		FClientLobbyData ClientData(LobbyId);
 
 		FAccountId User1(EOnlineServices::Null, 1);
@@ -454,7 +454,7 @@ bool FClientLobbyDataTest::RunTest(const FString& Parameters)
 	{
 		FLobbyEvents LobbyEvents;
 		FLobbyEventCapture EventCapture(LobbyEvents);
-		FOnlineLobbyIdHandle LobbyId(EOnlineServices::Null, 1);
+		FLobbyId LobbyId(EOnlineServices::Null, 1);
 		FClientLobbyData ClientData(LobbyId);
 
 		FAccountId User1(EOnlineServices::Null, 1);
@@ -597,7 +597,7 @@ bool FClientLobbyDataTest::RunTest(const FString& Parameters)
 	{
 		FLobbyEvents LobbyEvents;
 		FLobbyEventCapture EventCapture(LobbyEvents);
-		FOnlineLobbyIdHandle LobbyId(EOnlineServices::Null, 1);
+		FLobbyId LobbyId(EOnlineServices::Null, 1);
 		FClientLobbyData ClientData(LobbyId);
 
 		FAccountId User1(EOnlineServices::Null, 1);
@@ -776,7 +776,7 @@ bool FClientLobbyDataTest::RunTest(const FString& Parameters)
 	{
 		FLobbyEvents LobbyEvents;
 		FLobbyEventCapture EventCapture(LobbyEvents);
-		FOnlineLobbyIdHandle LobbyId(EOnlineServices::Null, 1);
+		FLobbyId LobbyId(EOnlineServices::Null, 1);
 		FClientLobbyData ClientData(LobbyId);
 
 		FAccountId User1(EOnlineServices::Null, 1);
@@ -1121,16 +1121,16 @@ TFuture<int> AwaitNextGameTick()
 	return AwaitSleepFor(0.f);
 }
 
-TFuture<TDefaultErrorResultInternal<FOnlineLobbyIdHandle>> AwaitInvitation(
+TFuture<TDefaultErrorResultInternal<FLobbyId>> AwaitInvitation(
 	FLobbyEvents& LobbyEvents,
 	FAccountId TargetAccountId,
-	FOnlineLobbyIdHandle LobbyId,
+	FLobbyId LobbyId,
 	float TimeoutSeconds)
 {
 	// todo: Make this nicer - current implementation has issue on shutdown.
 	struct AwaitInvitationState
 	{
-		TPromise<TDefaultErrorResultInternal<FOnlineLobbyIdHandle>> Promise;
+		TPromise<TDefaultErrorResultInternal<FLobbyId>> Promise;
 		FOnlineEventDelegateHandle OnLobbyInvitationAddedHandle;
 		FTSTicker::FDelegateHandle OnAwaitExpiredHandle;
 		bool bSignaled = false;
@@ -1674,7 +1674,7 @@ TOnlineAsyncOpHandle<FFunctionalTestLobbies> RunLobbyFunctionalTest(IAuth& AuthI
 		const FCreateLobby::Result& CreateLobbyResult = GetOpDataChecked<FCreateLobby::Result>(InAsyncOp, CreateLobbyKeyName);
 		return Private::AwaitInvitation(*LobbyEventsPtr, User2Info.AccountInfo->AccountId, CreateLobbyResult.Lobby->LobbyId, TestConfig->InvitationWaitSeconds);
 	})
-	.Then([](TOnlineAsyncOp<FFunctionalTestLobbies>& InAsyncOp, TDefaultErrorResultInternal<FOnlineLobbyIdHandle>&& Result)
+	.Then([](TOnlineAsyncOp<FFunctionalTestLobbies>& InAsyncOp, TDefaultErrorResultInternal<FLobbyId>&& Result)
 	{
 		if (Result.IsError())
 		{
diff --git a/Engine/Plugins/Online/Experimental/OnlineServices/Source/OnlineServicesCommon/Private/Online/LobbiesCommonTypes.cpp b/Engine/Plugins/Online/Experimental/OnlineServices/Source/OnlineServicesCommon/Private/Online/LobbiesCommonTypes.cpp
--- a/Engine/Plugins/Online/Experimental/OnlineServices/Source/OnlineServicesCommon/Private/Online/LobbiesCommonTypes.cpp
+++ b/Engine/Plugins/Online/Experimental/OnlineServices/Source/OnlineServicesCommon/Private/Online/LobbiesCommonTypes.cpp
@@ -97,7 +97,7 @@ bool FLobbySchemaRegistry::RegisterSchema(FLobbySchemaConfig LobbySchemaConfig)
 }
 
 
-FClientLobbyData::FClientLobbyData(FOnlineLobbyIdHandle LobbyId)
+FClientLobbyData::FClientLobbyData(FLobbyId LobbyId)
 	: PublicData(MakeShared<FLobby>())
 {
 	PublicData->LobbyId = LobbyId;
diff --git a/Engine/Plugins/Online/Experimental/OnlineServices/Source/OnlineServicesCommon/Public/Online/LobbiesCommonTypes.h b/Engine/Plugins/Online/Experimental/OnlineServices/Source/OnlineServicesCommon/Public/Online/LobbiesCommonTypes.h
--- a/Engine/Plugins/Online/Experimental/OnlineServices/Source/OnlineServicesCommon/Public/Online/LobbiesCommonTypes.h
+++ b/Engine/Plugins/Online/Experimental/OnlineServices/Source/OnlineServicesCommon/Public/Online/LobbiesCommonTypes.h
@@ -242,7 +242,7 @@ struct FApplyLobbyUpdateResult
 struct ONLINESERVICESCOMMON_API FClientLobbyData final
 {
 public:
-	FClientLobbyData(FOnlineLobbyIdHandle LobbyId);
+	FClientLobbyData(FLobbyId LobbyId);
 
 	TSharedRef<const FLobby> GetPublicDataPtr() const { return PublicData; }
 	const FLobby& GetPublicData() const { return *PublicData; }
diff --git a/Engine/Plugins/Online/Experimental/OnlineServices/Source/OnlineServicesInterface/Public/Online/Lobbies.h b/Engine/Plugins/Online/Experimental/OnlineServices/Source/OnlineServicesInterface/Public/Online/Lobbies.h
--- a/Engine/Plugins/Online/Experimental/OnlineServices/Source/OnlineServicesInterface/Public/Online/Lobbies.h
+++ b/Engine/Plugins/Online/Experimental/OnlineServices/Source/OnlineServicesInterface/Public/Online/Lobbies.h
@@ -64,7 +64,7 @@ struct FLobbyMember
 
 struct FLobby
 {
-	FOnlineLobbyIdHandle LobbyId;
+	FLobbyId LobbyId;
 	FAccountId OwnerAccountId;
 	FName LocalName;
 	FName SchemaName;
@@ -157,7 +157,7 @@ struct FFindLobbies
 		TOptional<FAccountId> TargetUser;
 
 		/** Find join info for the target lobby id. */
-		TOptional<FOnlineLobbyIdHandle> LobbyId;
+		TOptional<FLobbyId> LobbyId;
 	};
 
 	/** Output struct for Lobbies::FindLobbies */
@@ -196,7 +196,7 @@ struct FJoinLobby
 		FName LocalName;
 
 		/** The id of the lobby to be joined. */
-		FOnlineLobbyIdHandle LobbyId;
+		FLobbyId LobbyId;
 
 		/** Local users who will be joining the lobby. */
 		TArray<FJoinLobbyLocalUserData> LocalUsers;
@@ -221,7 +221,7 @@ struct FLeaveLobby
 		FAccountId LocalAccountId;
 
 		/** Id of the lobby to leave. */
-		FOnlineLobbyIdHandle LobbyId;
+		FLobbyId LobbyId;
 	};
 
 	/** Output struct for Lobbies::LeaveLobby */
@@ -241,7 +241,7 @@ struct FInviteLobbyMember
 		FAccountId LocalAccountId;
 
 		/** Id of the lobby for which the invitation will be sent. */
-		FOnlineLobbyIdHandle LobbyId;
+		FLobbyId LobbyId;
 
 		/** Id of the player who will be sent the invitation. */
 		FAccountId TargetAccountId;
@@ -264,7 +264,7 @@ struct FDeclineLobbyInvitation
 		FAccountId LocalAccountId;
 
 		/** Id of the lobby for which the invitations will be declined. */
-		FOnlineLobbyIdHandle LobbyId;
+		FLobbyId LobbyId;
 	};
 
 	/** Output struct for Lobbies::DeclineLobbyInvitation */
@@ -284,7 +284,7 @@ struct FKickLobbyMember
 		FAccountId LocalAccountId;
 
 		/** Id of the lobby. */
-		FOnlineLobbyIdHandle LobbyId;
+		FLobbyId LobbyId;
 
 		/** The target user to be kicked. */
 		FAccountId TargetAccountId;
@@ -307,7 +307,7 @@ struct FPromoteLobbyMember
 		FAccountId LocalAccountId;
 
 		/** Id of the lobby. */
-		FOnlineLobbyIdHandle LobbyId;
+		FLobbyId LobbyId;
 
 		/** The target user to be promoted to owner. */
 		FAccountId TargetAccountId;
@@ -345,7 +345,7 @@ struct FModifyLobbyJoinPolicy
 		FAccountId LocalAccountId;
 
 		/** Id of the lobby. */
-		FOnlineLobbyIdHandle LobbyId;
+		FLobbyId LobbyId;
 
 		/** The new join policy setting. */
 		ELobbyJoinPolicy JoinPolicy = ELobbyJoinPolicy::InvitationOnly;
@@ -368,7 +368,7 @@ struct FModifyLobbyAttributes
 		FAccountId LocalAccountId;
 
 		/** Id of the lobby. */
-		FOnlineLobbyIdHandle LobbyId;
+		FLobbyId LobbyId;
 
 		/** New or changed lobby attributes. */
 		TMap<FLobbyAttributeId, FLobbyVariant> MutatedAttributes;
@@ -394,7 +394,7 @@ struct FModifyLobbyMemberAttributes
 		FAccountId LocalAccountId;
 
 		/** Id of the lobby. */
-		FOnlineLobbyIdHandle LobbyId;
+		FLobbyId LobbyId;
 
 		/** New or changed lobby attributes. */
 		TMap<FLobbyAttributeId, FLobbyVariant> MutatedAttributes;
diff --git a/Engine/Plugins/Online/Experimental/OnlineServices/Source/OnlineServicesInterface/Public/Online/OnlineServices.h b/Engine/Plugins/Online/Experimental/OnlineServices/Source/OnlineServicesInterface/Public/Online/OnlineServices.h
--- a/Engine/Plugins/Online/Experimental/OnlineServices/Source/OnlineServicesInterface/Public/Online/OnlineServices.h
+++ b/Engine/Plugins/Online/Experimental/OnlineServices/Source/OnlineServicesInterface/Public/Online/OnlineServices.h
@@ -36,7 +36,7 @@ struct FGetResolvedConnectString
 	struct Params
 	{
 		FAccountId LocalAccountId;
-		FOnlineLobbyIdHandle LobbyId;
+		FLobbyId LobbyId;
 		FOnlineSessionIdHandle SessionId;
 		FName PortType;
 	};
diff --git a/Engine/Plugins/Online/Experimental/OnlineServicesEOSGS/Source/Private/Online/LobbiesEOSGS.cpp b/Engine/Plugins/Online/Experimental/OnlineServicesEOSGS/Source/Private/Online/LobbiesEOSGS.cpp
--- a/Engine/Plugins/Online/Experimental/OnlineServicesEOSGS/Source/Private/Online/LobbiesEOSGS.cpp
+++ b/Engine/Plugins/Online/Experimental/OnlineServicesEOSGS/Source/Private/Online/LobbiesEOSGS.cpp
@@ -1171,8 +1171,8 @@ void FLobbiesEOSGS::RemoveActiveLobby(FAccountId LocalAccountId, const TSharedRe
 
 void FLobbiesEOSGS::AddActiveInvite(const TSharedRef<FLobbyInviteDataEOS>& Invite)
 {
-	TMap<FOnlineLobbyIdHandle, TSharedRef<FLobbyInviteDataEOS>>& ActiveUserInvites = ActiveInvites.FindOrAdd(Invite->GetReceiver());
-	const FOnlineLobbyIdHandle LobbyId = Invite->GetLobbyData()->GetLobbyIdHandle();
+	TMap<FLobbyId, TSharedRef<FLobbyInviteDataEOS>>& ActiveUserInvites = ActiveInvites.FindOrAdd(Invite->GetReceiver());
+	const FLobbyId LobbyId = Invite->GetLobbyData()->GetLobbyIdHandle();
 
 	// Todo: Handle multiple invites for the same lobby.
 	if (ActiveUserInvites.Find(LobbyId) == nullptr)
@@ -1199,7 +1199,7 @@ void FLobbiesEOSGS::RemoveActiveInvite(const TSharedRef<FLobbyInviteDataEOS>& In
 		});
 }
 
-TSharedPtr<FLobbyInviteDataEOS> FLobbiesEOSGS::GetActiveInvite(FAccountId TargetUser, FOnlineLobbyIdHandle TargetLobbyId)
+TSharedPtr<FLobbyInviteDataEOS> FLobbiesEOSGS::GetActiveInvite(FAccountId TargetUser, FLobbyId TargetLobbyId)
 {
 	const TSharedRef<FLobbyInviteDataEOS>* Result = ActiveInvites.FindOrAdd(TargetUser).Find(TargetLobbyId);
 	return Result ? *Result : TSharedPtr<FLobbyInviteDataEOS>();
diff --git a/Engine/Plugins/Online/Experimental/OnlineServicesEOSGS/Source/Private/Online/LobbiesEOSGSTypes.cpp b/Engine/Plugins/Online/Experimental/OnlineServicesEOSGS/Source/Private/Online/LobbiesEOSGSTypes.cpp
--- a/Engine/Plugins/Online/Experimental/OnlineServicesEOSGS/Source/Private/Online/LobbiesEOSGSTypes.cpp
+++ b/Engine/Plugins/Online/Experimental/OnlineServicesEOSGS/Source/Private/Online/LobbiesEOSGSTypes.cpp
@@ -689,7 +689,7 @@ FLobbyDataEOS::FLobbyDataEOS(
 }
 
 TFuture<TDefaultErrorResultInternal<TSharedRef<FLobbyDataEOS>>> FLobbyDataEOS::Create(
-	FOnlineLobbyIdHandle LobbyIdHandle,
+	FLobbyId LobbyId,
 	const TSharedRef<FLobbyDetailsEOS>& LobbyDetails,
 	FUnregisterFn UnregisterFn)
 {
@@ -700,7 +700,7 @@ TFuture<TDefaultErrorResultInternal<TSharedRef<FLobbyDataEOS>>> FLobbyDataEOS::C
 	.Then(
 	[
 		Promise = MoveTemp(Promise),
-		LobbyIdHandle,
+		LobbyId,
 		LobbyDetails,
 		UnregisterFn = MoveTemp(UnregisterFn)
 	]
@@ -714,7 +714,7 @@ TFuture<TDefaultErrorResultInternal<TSharedRef<FLobbyDataEOS>>> FLobbyDataEOS::C
 		}
 
 		TSharedRef<FClientLobbySnapshot> LobbySnapshot = Future.Get().GetOkValue();
-		TSharedRef<FClientLobbyData> LobbyData = MakeShared<FClientLobbyData>(LobbyIdHandle);
+		TSharedRef<FClientLobbyData> LobbyData = MakeShared<FClientLobbyData>(LobbyId);
 
 		// Fetch member data and apply them to the lobby.
 		TMap<FAccountId, TSharedRef<FClientLobbyMemberSnapshot>> MemberSnapshots;
@@ -753,9 +753,9 @@ TSharedPtr<FLobbyDataEOS> FLobbyDataRegistryEOS::Find(EOS_LobbyId EOSLobbyId) co
 	return Result ? Result->Pin() : TSharedPtr<FLobbyDataEOS>();
 }
 
-TSharedPtr<FLobbyDataEOS> FLobbyDataRegistryEOS::Find(FOnlineLobbyIdHandle LobbyIdHandle) const
+TSharedPtr<FLobbyDataEOS> FLobbyDataRegistryEOS::Find(FLobbyId LobbyId) const
 {
-	const TWeakPtr<FLobbyDataEOS>* Result = LobbyIdHandleIndex.Find(LobbyIdHandle);
+	const TWeakPtr<FLobbyDataEOS>* Result = LobbyIdHandleIndex.Find(LobbyId);
 	return Result ? Result->Pin() : TSharedPtr<FLobbyDataEOS>();
 }
 
@@ -770,7 +770,7 @@ TFuture<TDefaultErrorResultInternal<TSharedRef<FLobbyDataEOS>>> FLobbyDataRegist
 	TPromise<TDefaultErrorResultInternal<TSharedRef<FLobbyDataEOS>>> Promise;
 	TFuture<TDefaultErrorResultInternal<TSharedRef<FLobbyDataEOS>>> Future = Promise.GetFuture();
 
-	const FOnlineLobbyIdHandle LobbyId = FOnlineLobbyIdHandle(EOnlineServices::Epic, NextHandleIndex++);
+	const FLobbyId LobbyId = FLobbyId(EOnlineServices::Epic, NextHandleIndex++);
 	FLobbyDataEOS::Create(LobbyId, LobbyDetails, MakeUnregisterFn())
 	.Then([WeakThis = AsWeak(), Promise = MoveTemp(Promise), LocalAccountId, LobbyDetails](TFuture<TDefaultErrorResultInternal<TSharedRef<FLobbyDataEOS>>>&& Future) mutable
 	{
@@ -795,9 +795,9 @@ void FLobbyDataRegistryEOS::Register(const TSharedRef<FLobbyDataEOS>& LobbyIdHan
 	LobbyIdHandleIndex.Add(LobbyIdHandleData->GetLobbyIdHandle(), LobbyIdHandleData);
 }
 
-void FLobbyDataRegistryEOS::Unregister(FOnlineLobbyIdHandle LobbyIdHandle)
+void FLobbyDataRegistryEOS::Unregister(FLobbyId LobbyId)
 {
-	if (TSharedPtr<FLobbyDataEOS> HandleData = Find(LobbyIdHandle))
+	if (TSharedPtr<FLobbyDataEOS> HandleData = Find(LobbyId))
 	{
 		LobbyIdIndex.Remove(HandleData->GetLobbyId());
 		LobbyIdHandleIndex.Remove(HandleData->GetLobbyIdHandle());
@@ -806,7 +806,7 @@ void FLobbyDataRegistryEOS::Unregister(FOnlineLobbyIdHandle LobbyIdHandle)
 
 FLobbyDataEOS::FUnregisterFn FLobbyDataRegistryEOS::MakeUnregisterFn()
 {
-	return [WeakThis = AsWeak()](FOnlineLobbyIdHandle LobbyId)
+	return [WeakThis = AsWeak()](FLobbyId LobbyId)
 	{
 		if (TSharedPtr<FLobbyDataRegistryEOS> StrongThis = WeakThis.Pin())
 		{
@@ -832,7 +832,7 @@ TFuture<TDefaultErrorResultInternal<TSharedRef<FLobbyInviteDataEOS>>> FLobbyInvi
 	TPromise<TDefaultErrorResultInternal<TSharedRef<FLobbyInviteDataEOS>>> Promise;
 	TFuture<TDefaultErrorResultInternal<TSharedRef<FLobbyInviteDataEOS>>> Future = Promise.GetFuture();
 
-	// Search for existing lobby data so that the LobbyIdHandle will match.
+	// Search for existing lobby data so that the LobbyId will match.
 	TSharedRef<FLobbyInviteIdEOS> InviteId = MakeShared<FLobbyInviteIdEOS>(InviteIdEOS);
 	LobbyDataRegistry->FindOrCreateFromLobbyDetails(LocalAccountId, LobbyDetails)
 	.Then([Promise = MoveTemp(Promise), InviteId, LocalAccountId, Sender, LobbyDetails](TFuture<TDefaultErrorResultInternal<TSharedRef<FLobbyDataEOS>>>&& Future) mutable
diff --git a/Engine/Plugins/Online/Experimental/OnlineServicesEOSGS/Source/Private/Online/LobbiesEOSGSTypes.h b/Engine/Plugins/Online/Experimental/OnlineServicesEOSGS/Source/Private/Online/LobbiesEOSGSTypes.h
--- a/Engine/Plugins/Online/Experimental/OnlineServicesEOSGS/Source/Private/Online/LobbiesEOSGSTypes.h
+++ b/Engine/Plugins/Online/Experimental/OnlineServicesEOSGS/Source/Private/Online/LobbiesEOSGSTypes.h
@@ -318,18 +318,18 @@ private:
 class FLobbyDataEOS final : public TSharedFromThis<FLobbyDataEOS>
 {
 public:
-	using FUnregisterFn = TFunction<void(FOnlineLobbyIdHandle)>;
+	using FUnregisterFn = TFunction<void(FLobbyId)>;
 
 	FLobbyDataEOS() = delete;
 
 	static TFuture<TDefaultErrorResultInternal<TSharedRef<FLobbyDataEOS>>> Create(
-		FOnlineLobbyIdHandle LobbyIdHandle,
+		FLobbyId LobbyId,
 		const TSharedRef<FLobbyDetailsEOS>& LobbyDetails,
 		FUnregisterFn UnregisterFn = FUnregisterFn());
 
 	~FLobbyDataEOS();
 
-	FOnlineLobbyIdHandle GetLobbyIdHandle() const { return ClientLobbyData->GetPublicData().LobbyId; }
+	FLobbyId GetLobbyIdHandle() const { return ClientLobbyData->GetPublicData().LobbyId; }
 	const TSharedRef<FClientLobbyData>& GetClientLobbyData() const { return ClientLobbyData; }
 	EOS_LobbyId GetLobbyIdEOS() const { return LobbyDetailsInfo->GetLobbyId(); }
 	const FString& GetLobbyId() const { return LobbyId; }
@@ -366,17 +366,17 @@ public:
 	FLobbyDataRegistryEOS(const TSharedRef<FLobbyPrerequisitesEOS>& Prerequisites);
 
 	TSharedPtr<FLobbyDataEOS> Find(EOS_LobbyId EOSLobbyId) const;
-	TSharedPtr<FLobbyDataEOS> Find(FOnlineLobbyIdHandle LobbyIdHandle) const;
+	TSharedPtr<FLobbyDataEOS> Find(FLobbyId LobbyId) const;
 	TFuture<TDefaultErrorResultInternal<TSharedRef<FLobbyDataEOS>>> FindOrCreateFromLobbyDetails(FAccountId LocalAccountId, const TSharedRef<FLobbyDetailsEOS>& LobbyDetails);
 
 private:
 	void Register(const TSharedRef<FLobbyDataEOS>& LobbyIdHandleData);
-	void Unregister(FOnlineLobbyIdHandle LobbyIdHandle);
+	void Unregister(FLobbyId LobbyId);
 	FLobbyDataEOS::FUnregisterFn MakeUnregisterFn();
 
 	TSharedRef<FLobbyPrerequisitesEOS> Prerequisites;
 	TMap<FString, TWeakPtr<FLobbyDataEOS>> LobbyIdIndex;
-	TMap<FOnlineLobbyIdHandle, TWeakPtr<FLobbyDataEOS>> LobbyIdHandleIndex;
+	TMap<FLobbyId, TWeakPtr<FLobbyDataEOS>> LobbyIdHandleIndex;
 	uint32 NextHandleIndex = 1;
 };
 
diff --git a/Engine/Plugins/Online/Experimental/OnlineServicesEOSGS/Source/Public/Online/LobbiesEOSGS.h b/Engine/Plugins/Online/Experimental/OnlineServicesEOSGS/Source/Public/Online/LobbiesEOSGS.h
--- a/Engine/Plugins/Online/Experimental/OnlineServicesEOSGS/Source/Public/Online/LobbiesEOSGS.h
+++ b/Engine/Plugins/Online/Experimental/OnlineServicesEOSGS/Source/Public/Online/LobbiesEOSGS.h
@@ -134,7 +134,7 @@ struct FLobbiesDeclineLobbyInvitationImpl
 		FAccountId LocalAccountId;
 
 		// Id of the lobby for which the invitations will be declined.
-		FOnlineLobbyIdHandle LobbyId;
+		FLobbyId LobbyId;
 	};
 
 	struct Result
@@ -299,7 +299,7 @@ protected:
 	// Todo: store list of invites per lobby.
 	void AddActiveInvite(const TSharedRef<FLobbyInviteDataEOS>& Invite);
 	void RemoveActiveInvite(const TSharedRef<FLobbyInviteDataEOS>& Invite);
-	TSharedPtr<FLobbyInviteDataEOS> GetActiveInvite(FAccountId TargetUser, FOnlineLobbyIdHandle TargetLobbyId);
+	TSharedPtr<FLobbyInviteDataEOS> GetActiveInvite(FAccountId TargetUser, FLobbyId TargetLobbyId);
 
 	// LobbyData will be fetched from the operation data if not set in Params.
 	TFuture<TDefaultErrorResult<FLobbiesJoinLobbyImpl>> JoinLobbyImpl(FLobbiesJoinLobbyImpl::Params&& Params);
@@ -326,7 +326,7 @@ protected:
 	TSharedPtr<FLobbyDataRegistryEOS> LobbyDataRegistry;
 
 	TMap<FAccountId, TSet<TSharedRef<FLobbyDataEOS>>> ActiveLobbies;
-	TMap<FAccountId, TMap<FOnlineLobbyIdHandle, TSharedRef<FLobbyInviteDataEOS>>> ActiveInvites;
+	TMap<FAccountId, TMap<FLobbyId, TSharedRef<FLobbyInviteDataEOS>>> ActiveInvites;
 	TMap<FAccountId, TSharedRef<FLobbySearchEOS>> ActiveSearchResults;
 };
 
diff --git a/Engine/Plugins/Online/Experimental/OnlineServicesNull/Source/Private/Online/LobbyRegistryNull.cpp b/Engine/Plugins/Online/Experimental/OnlineServicesNull/Source/Private/Online/LobbyRegistryNull.cpp
--- a/Engine/Plugins/Online/Experimental/OnlineServicesNull/Source/Private/Online/LobbyRegistryNull.cpp
+++ b/Engine/Plugins/Online/Experimental/OnlineServicesNull/Source/Private/Online/LobbyRegistryNull.cpp
@@ -9,44 +9,44 @@ FOnlineLobbyIdRegistryNull& FOnlineLobbyIdRegistryNull::Get()
 	return Instance;
 }
 
-const FOnlineLobbyIdHandle* FOnlineLobbyIdRegistryNull::Find(FString LobbyId)
+const FLobbyId* FOnlineLobbyIdRegistryNull::Find(FString LobbyIdStr)
 {
-	return StringToId.Find(LobbyId);
+	return StringToId.Find(LobbyIdStr);
 }
 
-FOnlineLobbyIdHandle FOnlineLobbyIdRegistryNull::FindOrAdd(FString LobbyId)
+FLobbyId FOnlineLobbyIdRegistryNull::FindOrAdd(FString LobbyIdStr)
 {
-	const FOnlineLobbyIdHandle* Entry = StringToId.Find(LobbyId);
+	const FLobbyId* Entry = StringToId.Find(LobbyIdStr);
 	if (Entry)
 	{
 		return *Entry;
 	}
 
-	Ids.Add(LobbyId);
-	FOnlineLobbyIdHandle Handle(EOnlineServices::Null, Ids.Num());
-	StringToId.Add(LobbyId, Handle);
+	Ids.Add(LobbyIdStr);
+	FLobbyId LobbyId(EOnlineServices::Null, Ids.Num());
+	StringToId.Add(LobbyIdStr, LobbyId);
 
-	return Handle;
+	return LobbyId;
 }
 
-UE::Online::FOnlineLobbyIdHandle FOnlineLobbyIdRegistryNull::GetNext()
+UE::Online::FLobbyId FOnlineLobbyIdRegistryNull::GetNext()
 {
 	return FindOrAdd(FGuid::NewGuid().ToString());
 }
 
 
-const FString* FOnlineLobbyIdRegistryNull::GetInternal(const FOnlineLobbyIdHandle& Handle) const
+const FString* FOnlineLobbyIdRegistryNull::GetInternal(const FLobbyId& LobbyId) const
 {
-	if (Handle.GetOnlineServicesType() == EOnlineServices::Null && Handle.GetHandle() <= (uint32)Ids.Num())
+	if (LobbyId.GetOnlineServicesType() == EOnlineServices::Null && LobbyId.GetHandle() <= (uint32)Ids.Num())
 	{
-		return &Ids[Handle.GetHandle()-1];
+		return &Ids[LobbyId.GetHandle()-1];
 	}
 	return nullptr;
 }
 
-FString FOnlineLobbyIdRegistryNull::ToLogString(const FOnlineLobbyIdHandle& Handle) const
+FString FOnlineLobbyIdRegistryNull::ToLogString(const FLobbyId& LobbyId) const
 {
-	if (const FString* Id = GetInternal(Handle))
+	if (const FString* Id = GetInternal(LobbyId))
 	{
 		return *Id;
 	}
@@ -55,9 +55,9 @@ FString FOnlineLobbyIdRegistryNull::ToLogString(const FOnlineLobbyIdHandle& Hand
 }
 
 
-TArray<uint8> FOnlineLobbyIdRegistryNull::ToReplicationData(const FOnlineLobbyIdHandle& Handle) const
+TArray<uint8> FOnlineLobbyIdRegistryNull::ToReplicationData(const FLobbyId& LobbyId) const
 {
-	if (const FString* Id = GetInternal(Handle))
+	if (const FString* Id = GetInternal(LobbyId))
 	{
 		TArray<uint8> ReplicationData;
 		ReplicationData.Reserve(Id->Len());
@@ -68,14 +68,14 @@ TArray<uint8> FOnlineLobbyIdRegistryNull::ToReplicationData(const FOnlineLobbyId
 	return TArray<uint8>();;
 }
 
-FOnlineLobbyIdHandle FOnlineLobbyIdRegistryNull::FromReplicationData(const TArray<uint8>& ReplicationData)
+FLobbyId FOnlineLobbyIdRegistryNull::FromReplicationData(const TArray<uint8>& ReplicationData)
 {
 	FString Result = BytesToString(ReplicationData.GetData(), ReplicationData.Num());
 	if (Result.Len() > 0)
 	{
 		return FindOrAdd(Result);
 	}
-	return FOnlineLobbyIdHandle();
+	return FLobbyId();
 }
 
 } // namespace UE::Online
diff --git a/Engine/Plugins/Online/Experimental/OnlineServicesNull/Source/Public/Online/LobbiesNull.h b/Engine/Plugins/Online/Experimental/OnlineServicesNull/Source/Public/Online/LobbiesNull.h
--- a/Engine/Plugins/Online/Experimental/OnlineServicesNull/Source/Public/Online/LobbiesNull.h
+++ b/Engine/Plugins/Online/Experimental/OnlineServicesNull/Source/Public/Online/LobbiesNull.h
@@ -62,7 +62,7 @@ private:
 	FOnlineServicesNull& Services;
 	
 	TMap<FName, TSharedRef<FLobbyNull>> NamedLobbies;
-	TMap<FOnlineLobbyIdHandle, TSharedRef<FLobbyNull>> AllLobbies; // todo: this should be moved to the lobby id registry (replication made this tricky to get implemented asap)
+	TMap<FLobbyId, TSharedRef<FLobbyNull>> AllLobbies; // todo: this should be moved to the lobby id registry (replication made this tricky to get implemented asap)
 	TSharedPtr<FFindLobbies::Result> CurrentLobbySearch;
 	TSharedPtr<TOnlineAsyncOp<FFindLobbies>> CurrentLobbySearchHandle;
 	FLANSession LANSessionManager;
@@ -98,7 +98,7 @@ private:
 		}
 	}
 
-	TSharedPtr<FLobbyNull> GetLobby(const FOnlineLobbyIdHandle LobbyId)
+	TSharedPtr<FLobbyNull> GetLobby(const FLobbyId LobbyId)
 	{
 		TSharedRef<FLobbyNull>* Lobby = AllLobbies.Find(LobbyId);
 		if(Lobby)
@@ -126,7 +126,7 @@ private:
 	}
 
 	// returns true if a named lobby was removed
-	bool RemoveLobbyFromId(const FOnlineLobbyIdHandle& LobbyId)
+	bool RemoveLobbyFromId(const FLobbyId& LobbyId)
 	{
 		TSharedPtr<FLobbyNull> Lobby = GetLobby(LobbyId);
 		if(Lobby)
diff --git a/Engine/Plugins/Online/Experimental/OnlineServicesNull/Source/Public/Online/LobbyRegistryNull.h b/Engine/Plugins/Online/Experimental/OnlineServicesNull/Source/Public/Online/LobbyRegistryNull.h
--- a/Engine/Plugins/Online/Experimental/OnlineServicesNull/Source/Public/Online/LobbyRegistryNull.h
+++ b/Engine/Plugins/Online/Experimental/OnlineServicesNull/Source/Public/Online/LobbyRegistryNull.h
@@ -12,22 +12,22 @@ class FOnlineLobbyIdRegistryNull : public IOnlineLobbyIdRegistry
 public:
 	static FOnlineLobbyIdRegistryNull& Get();
 
-	const FOnlineLobbyIdHandle* Find(FString LobbyId);
-	FOnlineLobbyIdHandle FindOrAdd(FString LobbyId);
-	FOnlineLobbyIdHandle GetNext();
+	const FLobbyId* Find(FString LobbyIdStr);
+	FLobbyId FindOrAdd(FString LobbyIdStr);
+	FLobbyId GetNext();
 
 	// Begin IOnlineAccountIdRegistry
-	virtual FString ToLogString(const FOnlineLobbyIdHandle& Handle) const override;
-	virtual TArray<uint8> ToReplicationData(const FOnlineLobbyIdHandle& Handle) const override;
-	virtual FOnlineLobbyIdHandle FromReplicationData(const TArray<uint8>& ReplicationString) override;
+	virtual FString ToLogString(const FLobbyId& LobbyId) const override;
+	virtual TArray<uint8> ToReplicationData(const FLobbyId& LobbyId) const override;
+	virtual FLobbyId FromReplicationData(const TArray<uint8>& ReplicationString) override;
 	// End IOnlineAccountIdRegistry
 
 	virtual ~FOnlineLobbyIdRegistryNull() = default;
 
 private:
-	const FString* GetInternal(const FOnlineLobbyIdHandle& Handle) const;
+	const FString* GetInternal(const FLobbyId& LobbyId) const;
 	TArray<FString> Ids;
-	TMap<FString, FOnlineLobbyIdHandle> StringToId; 
+	TMap<FString, FLobbyId> StringToId; 
 
 };
 
diff --git a/Engine/Source/Runtime/CoreOnline/Private/Online/CoreOnline.cpp b/Engine/Source/Runtime/CoreOnline/Private/Online/CoreOnline.cpp
--- a/Engine/Source/Runtime/CoreOnline/Private/Online/CoreOnline.cpp
+++ b/Engine/Source/Runtime/CoreOnline/Private/Online/CoreOnline.cpp
@@ -410,7 +410,7 @@ FString ToLogString(const FAccountId& Id)
 	return FString::Printf(TEXT("%s:%d (%s)"), LexToString(Id.GetOnlineServicesType()), Id.GetHandle(), *FOnlineIdRegistryRegistry::Get().ToLogString(Id));
 }
 
-FString ToLogString(const FOnlineLobbyIdHandle& Id)
+FString ToLogString(const FLobbyId& Id)
 {
 	// todo: use registry.
 	return FString::Printf(TEXT("%s:%d"), LexToString(Id.GetOnlineServicesType()), Id.GetHandle());
diff --git a/Engine/Source/Runtime/CoreOnline/Public/Online/CoreOnline.h b/Engine/Source/Runtime/CoreOnline/Public/Online/CoreOnline.h
--- a/Engine/Source/Runtime/CoreOnline/Public/Online/CoreOnline.h
+++ b/Engine/Source/Runtime/CoreOnline/Public/Online/CoreOnline.h
@@ -307,14 +307,14 @@ private:
 };
 
 using FAccountId = TOnlineIdHandle<OnlineIdHandleTags::FAccount>;
-using FOnlineLobbyIdHandle = TOnlineIdHandle<OnlineIdHandleTags::FLobby>;
+using FLobbyId = TOnlineIdHandle<OnlineIdHandleTags::FLobby>;
 using FOnlineSessionIdHandle = TOnlineIdHandle<OnlineIdHandleTags::FSession>;
 using FOnlineSessionInviteIdHandle = TOnlineIdHandle<OnlineIdHandleTags::FSessionInvite>;
 using FOnlineVerifiedAuthTicketIdHandle = TOnlineIdHandle<OnlineIdHandleTags::FVerifiedAuthTicket>;
 using FOnlineVerifiedAuthSessionIdHandle = TOnlineIdHandle<OnlineIdHandleTags::FVerifiedAuthSession>;
 
 COREONLINE_API FString ToLogString(const FAccountId& Id);
-COREONLINE_API FString ToLogString(const FOnlineLobbyIdHandle& Id);
+COREONLINE_API FString ToLogString(const FLobbyId& Id);
 COREONLINE_API FString ToLogString(const FOnlineSessionIdHandle& Id);
 COREONLINE_API FString ToLogString(const FOnlineSessionInviteIdHandle& Id);
 COREONLINE_API FString ToLogString(const FOnlineVerifiedAuthTicketIdHandle& Id);
