Remove multiple user actions from lobbies interface.
Add better lobbies logging and error handling.

diff --git a/Engine/Plugins/Online/Experimental/OnlineServices/Source/OnlineServicesCommon/Private/Online/LobbiesCommonTests.cpp b/Engine/Plugins/Online/Experimental/OnlineServices/Source/OnlineServicesCommon/Private/Online/LobbiesCommonTests.cpp
--- a/Engine/Plugins/Online/Experimental/OnlineServices/Source/OnlineServicesCommon/Private/Online/LobbiesCommonTests.cpp
+++ b/Engine/Plugins/Online/Experimental/OnlineServices/Source/OnlineServicesCommon/Private/Online/LobbiesCommonTests.cpp
@@ -1271,8 +1271,8 @@ TFuture<TOnlineResult<FFunctionalTestLogoutUser>> FunctionalTestLogoutUser(FFunc
 	TOnlineResult<FAuthGetLocalOnlineUserByPlatformUserId> Result = Params.AuthInterface->GetLocalOnlineUserByPlatformUserId(MoveTemp(GetAccountParams));
 	if (Result.IsError())
 	{
-		// Ignore errors for now. This function should ignore logged out users.
-		if (Result.GetErrorValue() == Errors::Unknown())
+		// Treat invalid user as success.
+		if (Result.GetErrorValue() == Errors::InvalidUser())
 		{
 			return MakeFulfilledPromise<TOnlineResult<FFunctionalTestLogoutUser>>(FFunctionalTestLogoutUser::Result{}).GetFuture();
 		}
@@ -1474,7 +1474,7 @@ TOnlineAsyncOpHandle<FFunctionalTestLobbies> RunLobbyFunctionalTest(IAuth& AuthI
 		Params.MaxMembers = 2;
 		Params.JoinPolicy = ELobbyJoinPolicy::InvitationOnly;
 		Params.Attributes.Add(TEXT("test_attribute_key"), TEXT("test_attribute_value"));
-		Params.LocalUsers.Emplace(FJoinLobbyLocalUserData{User1Info.AccountInfo->AccountId, {{TEXT("test_attribute_key"), TEXT("test_attribute_value")}}});
+		Params.UserAttributes.Add(TEXT("test_attribute_key"), TEXT("test_attribute_value"));
 		return Params;
 	})
 	.Then(Private::CaptureOperationStepResult<FCreateLobby>(*Op, CreateLobbyKeyName, Private::FBindOperation<FCreateLobby>(&LobbiesCommon, &FLobbiesCommon::CreateLobby)))
@@ -1637,7 +1637,6 @@ TOnlineAsyncOpHandle<FFunctionalTestLobbies> RunLobbyFunctionalTest(IAuth& AuthI
 		Params.MaxMembers = 2;
 		Params.JoinPolicy = ELobbyJoinPolicy::InvitationOnly;
 		//Params.Attributes;
-		Params.LocalUsers.Emplace(FJoinLobbyLocalUserData{User1Info.AccountInfo->AccountId, {}});
 		return Params;
 	})
 	.Then(Private::CaptureOperationStepResult<FCreateLobby>(*Op, CreateLobbyKeyName, Private::FBindOperation<FCreateLobby>(&LobbiesCommon, &FLobbiesCommon::CreateLobby)))
@@ -1693,7 +1692,6 @@ TOnlineAsyncOpHandle<FFunctionalTestLobbies> RunLobbyFunctionalTest(IAuth& AuthI
 		Params.LocalAccountId = User2Info.AccountInfo->AccountId;
 		Params.LocalName = TEXT("test");
 		Params.LobbyId = CreateLobbyResult.Lobby->LobbyId;
-		Params.LocalUsers.Emplace(FJoinLobbyLocalUserData{User2Info.AccountInfo->AccountId, {}});
 		return Params;
 	})
 	.Then(Private::ConsumeOperationStepResult<FJoinLobby>(*Op, Private::FBindOperation<FJoinLobby>(&LobbiesCommon, &FLobbiesCommon::JoinLobby)))
@@ -1786,7 +1784,6 @@ TOnlineAsyncOpHandle<FFunctionalTestLobbies> RunLobbyFunctionalTest(IAuth& AuthI
 		Params.MaxMembers = 2;
 		Params.JoinPolicy = ELobbyJoinPolicy::PublicAdvertised;
 		Params.Attributes = {{ TEXT("LobbyCreateTime"), SearchParams->LobbyCreateTime }};
-		Params.LocalUsers.Emplace(FJoinLobbyLocalUserData{User1Info.AccountInfo->AccountId, {}});
 		return Params;
 	})
 	.Then(Private::CaptureOperationStepResult<FCreateLobby>(*Op, CreateLobbyKeyName, Private::FBindOperation<FCreateLobby>(&LobbiesCommon, &FLobbiesCommon::CreateLobby)))
@@ -1839,7 +1836,6 @@ TOnlineAsyncOpHandle<FFunctionalTestLobbies> RunLobbyFunctionalTest(IAuth& AuthI
 		Params.LocalAccountId = User2Info.AccountInfo->AccountId;
 		Params.LocalName = TEXT("test");
 		Params.LobbyId = FindResults.Lobbies[0]->LobbyId;
-		Params.LocalUsers.Emplace(FJoinLobbyLocalUserData{User2Info.AccountInfo->AccountId, {}});
 		return Params;
 	})
 	.Then(Private::ConsumeOperationStepResult<FJoinLobby>(*Op, Private::FBindOperation<FJoinLobby>(&LobbiesCommon, &FLobbiesCommon::JoinLobby)))
@@ -1904,7 +1900,8 @@ TOnlineAsyncOpHandle<FFunctionalTestLobbies> RunLobbyFunctionalTest(IAuth& AuthI
 	//----------------------------------------------------------------------------------------------
 	// Complete
 	//----------------------------------------------------------------------------------------------
-
+	// Complete the test by logging out all users.
+	.Then(Private::ConsumeStepResult<Private::FFunctionalTestLogoutAllUsers>(*Op, &Private::FunctionalTestLogoutAllUsers, Private::FFunctionalTestLogoutAllUsers::Params{ &AuthInterface }))
 	.Then([](TOnlineAsyncOp<FFunctionalTestLobbies>& InAsyncOp)
 	{
 		InAsyncOp.SetResult(FFunctionalTestLobbies::Result{});
diff --git a/Engine/Plugins/Online/Experimental/OnlineServices/Source/OnlineServicesInterface/Public/Online/Lobbies.h b/Engine/Plugins/Online/Experimental/OnlineServices/Source/OnlineServicesInterface/Public/Online/Lobbies.h
--- a/Engine/Plugins/Online/Experimental/OnlineServices/Source/OnlineServicesInterface/Public/Online/Lobbies.h
+++ b/Engine/Plugins/Online/Experimental/OnlineServices/Source/OnlineServicesInterface/Public/Online/Lobbies.h
@@ -74,15 +74,6 @@ struct FLobby
 	TMap<FAccountId, TSharedRef<const FLobbyMember>> Members;
 };
 
-struct FJoinLobbyLocalUserData
-{
-	/** Local users who will be joining the lobby. */
-	FAccountId LocalAccountId;
-
-	/** Initial attributes. */
-	TMap<FLobbyAttributeId, FLobbyVariant> Attributes;
-};
-
 struct FFindLobbySearchFilter
 {
 	/** The name of the attribute to be compared. */
@@ -122,8 +113,8 @@ struct FCreateLobby
 		/** Initial attributes. */
 		TMap<FLobbyAttributeId, FLobbyVariant> Attributes;
 
-		/** Local users who will be joining the lobby. */
-		TArray<FJoinLobbyLocalUserData> LocalUsers;
+		/** Initial user attributes. */
+		TMap<FLobbyAttributeId, FLobbyVariant> UserAttributes;
 	};
 
 	/** Output struct for Lobbies::CreateLobby */
@@ -198,8 +189,8 @@ struct FJoinLobby
 		/** The id of the lobby to be joined. */
 		FLobbyId LobbyId;
 
-		/** Local users who will be joining the lobby. */
-		TArray<FJoinLobbyLocalUserData> LocalUsers;
+		/** Initial user attributes. */
+		TMap<FLobbyAttributeId, FLobbyVariant> UserAttributes;
 	};
 
 	/** Output struct for Lobbies::JoinLobby */
@@ -808,11 +799,6 @@ BEGIN_ONLINE_STRUCT_META(FLobby)
 	ONLINE_STRUCT_FIELD(FLobby, JoinPolicy)
 END_ONLINE_STRUCT_META()
 
-BEGIN_ONLINE_STRUCT_META(FJoinLobbyLocalUserData)
-	ONLINE_STRUCT_FIELD(FJoinLobbyLocalUserData, LocalAccountId),
-	ONLINE_STRUCT_FIELD(FJoinLobbyLocalUserData, Attributes)
-END_ONLINE_STRUCT_META()
-
 BEGIN_ONLINE_STRUCT_META(FCreateLobby::Params)
 	ONLINE_STRUCT_FIELD(FCreateLobby::Params, LocalAccountId),
 	ONLINE_STRUCT_FIELD(FCreateLobby::Params, LocalName),
@@ -820,7 +806,7 @@ BEGIN_ONLINE_STRUCT_META(FCreateLobby::Params)
 	ONLINE_STRUCT_FIELD(FCreateLobby::Params, MaxMembers),
 	ONLINE_STRUCT_FIELD(FCreateLobby::Params, JoinPolicy),
 	ONLINE_STRUCT_FIELD(FCreateLobby::Params, Attributes),
-	ONLINE_STRUCT_FIELD(FCreateLobby::Params, LocalUsers)
+	ONLINE_STRUCT_FIELD(FCreateLobby::Params, UserAttributes)
 END_ONLINE_STRUCT_META()
 
 BEGIN_ONLINE_STRUCT_META(FCreateLobby::Result)
diff --git a/Engine/Plugins/Online/Experimental/OnlineServicesEOSGS/Source/Private/Online/LobbiesEOSGS.cpp b/Engine/Plugins/Online/Experimental/OnlineServicesEOSGS/Source/Private/Online/LobbiesEOSGS.cpp
--- a/Engine/Plugins/Online/Experimental/OnlineServicesEOSGS/Source/Private/Online/LobbiesEOSGS.cpp
+++ b/Engine/Plugins/Online/Experimental/OnlineServicesEOSGS/Source/Private/Online/LobbiesEOSGS.cpp
@@ -19,6 +19,8 @@ static const FString LobbyDataKeyName = TEXT("LobbyData");
 static const FString LobbyDetailsKeyName = TEXT("LobbyDetails");
 static const FString LobbyChangesKeyName = TEXT("LobbyChanges");
 static const FString LobbyMemberChangesKeyName = TEXT("LobbyMemberChanges");
+static const FString LobbyIdStringKeyName = TEXT("LobbyIdString");
+static const FString LobbySearchKeyName = TEXT("LobbySearch");
 
 static const int32 MaxAttributeSize = 1000;
 
@@ -59,48 +61,59 @@ void FLobbiesEOSGS::PreShutdown()
 
 TOnlineAsyncOpHandle<FCreateLobby> FLobbiesEOSGS::CreateLobby(FCreateLobby::Params&& InParams)
 {
-	TOnlineAsyncOpRef<FCreateLobby> Op = GetOp<FCreateLobby>(MoveTemp(InParams));
-	const FCreateLobby::Params& Params = Op->GetParams();
-
-	if (!Services.Get<FAuthEOSGS>()->IsLoggedIn(Params.LocalAccountId))
-	{
-		Op->SetError(Errors::InvalidUser());
-		return Op->GetHandle();
-	}
-
+	// Helper lambda for simplifying error handling.
 	auto DestroyLobbyDuringCreate = [this](
 		TOnlineAsyncOp<FCreateLobby>& InAsyncOp,
-		const TSharedPtr<FLobbyDataEOS>& LobbyData,
 		FAccountId LocalAccountId,
+		const FString& LobbyIdString,
 		FOnlineError ErrorResult) -> TFuture<void>
 	{
 		FLobbiesDestroyLobbyImpl::Params DestroyLobbyParams;
-		DestroyLobbyParams.LobbyData = LobbyData;
+		DestroyLobbyParams.LobbyIdString = LobbyIdString;
 		DestroyLobbyParams.LocalAccountId = LocalAccountId;
 
 		TPromise<void> Promise;
 		auto Future = Promise.GetFuture();
 
 		DestroyLobbyImpl(MoveTemp(DestroyLobbyParams))
-		.Then([this, InAsyncOp = InAsyncOp.AsShared(), ErrorResult = MoveTemp(ErrorResult), Promise = MoveTemp(Promise)](TFuture<TDefaultErrorResult<FLobbiesDestroyLobbyImpl>>&& Future) mutable
+		.Next(
+		[
+			InAsyncOp = InAsyncOp.AsShared(),
+			LocalAccountId,
+			LobbyIdString,
+			ErrorResult = MoveTemp(ErrorResult),
+			Promise = MoveTemp(Promise)
+		](TDefaultErrorResult<FLobbiesDestroyLobbyImpl>&& Result) mutable
 		{
-			if (Future.Get().IsError())
+			if (Result.IsError())
 			{
-				// Todo: complain about having an error while handling an error.
+				UE_LOG(LogTemp, Warning, TEXT("[FLobbiesEOSGS::CreateLobby] DestroyLobbyImpl Failed: User[%s], Lobby[%s], Result[%s]"),
+					*ToLogString(LocalAccountId), *LobbyIdString, *Result.GetErrorValue().GetLogString());
 			}
 
-			// Todo: Errors.
-			InAsyncOp->SetError(Errors::Unknown(MoveTemp(ErrorResult)));
+			// Set operation error to the passed in error value.
+			InAsyncOp->SetError(MoveTemp(ErrorResult));
 			Promise.EmplaceValue();
 		});
 
 		return Future;
 	};
 
-	Op->Then([this](TOnlineAsyncOp<FCreateLobby>& InAsyncOp, TPromise<const EOS_Lobby_CreateLobbyCallbackInfo*>&& Promise)
+	TOnlineAsyncOpRef<FCreateLobby> Op = GetOp<FCreateLobby>(MoveTemp(InParams));
+	const FCreateLobby::Params& Params = Op->GetParams();
+
+	// Check prerequisites.
+	if (!Services.GetAuthInterface()->IsLoggedIn(Params.LocalAccountId))
 	{
-		// Step 1: Call create lobby.
+		UE_LOG(LogTemp, Warning, TEXT("[FLobbiesEOSGS::CreateLobby] Failed: User is not logged in. User[%s]"), *ToLogString(Params.LocalAccountId));
+		Op->SetError(Errors::InvalidUser());
+		return Op->GetHandle();
+	}
 
+	// Start operation.
+	// Step 1: Call create lobby.
+	Op->Then([this](TOnlineAsyncOp<FCreateLobby>& InAsyncOp, TPromise<const EOS_Lobby_CreateLobbyCallbackInfo*>&& Promise)
+	{
 		const FCreateLobby::Params& Params = InAsyncOp.GetParams();
 		const FLobbyBucketIdTranslator<ELobbyTranslationType::ToService> BucketTanslator(LobbyPrerequisites->BucketId);
 
@@ -111,7 +124,8 @@ TOnlineAsyncOpHandle<FCreateLobby> FLobbiesEOSGS::CreateLobby(FCreateLobby::Para
 		CreateLobbyOptions.ApiVersion = EOS_LOBBY_CREATELOBBY_API_LATEST;
 		CreateLobbyOptions.LocalUserId = GetProductUserIdChecked(Params.LocalAccountId);
 		CreateLobbyOptions.MaxLobbyMembers = Params.MaxMembers;
-		CreateLobbyOptions.PermissionLevel = TranslateJoinPolicy(ELobbyJoinPolicy::InvitationOnly);
+		// Prevent lobby from appearing in search results until fully created with all user attributes.
+		CreateLobbyOptions.PermissionLevel = EOS_ELobbyPermissionLevel::EOS_LPL_INVITEONLY;
 		CreateLobbyOptions.bPresenceEnabled = false; // todo: handle
 		CreateLobbyOptions.bAllowInvites = true; // todo: handle
 		CreateLobbyOptions.BucketId = BucketTanslator.GetBucketIdEOS();
@@ -121,118 +135,121 @@ TOnlineAsyncOpHandle<FCreateLobby> FLobbiesEOSGS::CreateLobby(FCreateLobby::Para
 
 		EOS_Async(EOS_Lobby_CreateLobby, LobbyPrerequisites->LobbyInterfaceHandle, CreateLobbyOptions, MoveTemp(Promise)); 
 	})
+	// Step 2: Handle errors and attach the lobby id to the operation data.
 	.Then([this](TOnlineAsyncOp<FCreateLobby>& InAsyncOp, const EOS_Lobby_CreateLobbyCallbackInfo* Data)
 	{
-		// Step 2: Start creating the lobby data from the EOS lobby details object.
-
 		const FCreateLobby::Params& Params = InAsyncOp.GetParams();
 
 		if (Data->ResultCode != EOS_EResult::EOS_Success)
 		{
-			// Todo: errors
+			UE_LOG(LogTemp, Warning, TEXT("[FLobbiesEOSGS::CreateLobby] EOS_Lobby_CreateLobby Failed: User[%s], Result[%s]"),
+				*ToLogString(Params.LocalAccountId), *LexToString(Data->ResultCode));
 			InAsyncOp.SetError(Errors::FromEOSResult(Data->ResultCode));
-			return MakeFulfilledPromise<TDefaultErrorResultInternal<TSharedRef<FLobbyDataEOS>>>().GetFuture();
 		}
+		else
+		{
+			// Store the lobby id string on the operation data. This is used to destroy the lobby when creating the lobby structures fails.
+			InAsyncOp.Data.Set<FString>(LobbyIdStringKeyName, FString(UTF8_TO_TCHAR(Data->LobbyId)));
+		}
+	})
+	// Step 3: Create the lobby details object from the lobby id.
+	.Then([this, DestroyLobbyDuringCreate](TOnlineAsyncOp<FCreateLobby>& InAsyncOp)
+	{
+		const FCreateLobby::Params& Params = InAsyncOp.GetParams();
+		const FString& LobbyIdString = GetOpDataChecked<FString>(InAsyncOp, LobbyIdStringKeyName);
 
-		TDefaultErrorResultInternal<TSharedRef<FLobbyDetailsEOS>> LobbyDetailsResult =
-			FLobbyDetailsEOS::CreateFromLobbyId(LobbyPrerequisites.ToSharedRef(), Params.LocalAccountId, Data->LobbyId);
-		if (LobbyDetailsResult.IsError())
+		// Try to create the lobby details object.
+		TDefaultErrorResultInternal<TSharedRef<FLobbyDetailsEOS>> Result =
+			FLobbyDetailsEOS::CreateFromLobbyId(LobbyPrerequisites.ToSharedRef(), Params.LocalAccountId, LobbyIdString);
+		if (Result.IsError())
 		{
-			// Todo: manually call eos destroy lobby here.
+			UE_LOG(LogTemp, Warning, TEXT("[FLobbiesEOSGS::CreateLobby] FLobbyDetailsEOS::CreateFromLobbyId Failed: User[%s], Lobby[%s], Result[%s]"),
+				*ToLogString(Params.LocalAccountId), *LobbyIdString, *Result.GetErrorValue().GetLogString());
 
-			// Todo: errors
-			InAsyncOp.SetError(MoveTemp(LobbyDetailsResult.GetErrorValue()));
-			return MakeFulfilledPromise<TDefaultErrorResultInternal<TSharedRef<FLobbyDataEOS>>>().GetFuture();
+			return DestroyLobbyDuringCreate(InAsyncOp, Params.LocalAccountId, LobbyIdString, MoveTemp(Result.GetErrorValue()));
+		}
+		else
+		{
+			// Attach the lobby details to the operation.
+			InAsyncOp.Data.Set<TSharedRef<FLobbyDetailsEOS>>(LobbyDetailsKeyName, Result.GetOkValue());
+			return MakeFulfilledPromise<void>().GetFuture();
 		}
-		
-		return LobbyDataRegistry->FindOrCreateFromLobbyDetails(InAsyncOp.GetParams().LocalAccountId, LobbyDetailsResult.GetOkValue());
 	})
+	// Step 4: Create the lobby data object from the lobby details.
+	.Then([this](TOnlineAsyncOp<FCreateLobby>& InAsyncOp)
+	{
+		TSharedRef<FLobbyDetailsEOS> LobbyDetails = GetOpDataChecked<TSharedRef<FLobbyDetailsEOS>>(InAsyncOp, LobbyDetailsKeyName);
+		return LobbyDataRegistry->FindOrCreateFromLobbyDetails(InAsyncOp.GetParams().LocalAccountId, LobbyDetails);
+	})
+	// Step 5: Handle errors and store the lobby data on the async op properties.
 	.Then([this, DestroyLobbyDuringCreate](TOnlineAsyncOp<FCreateLobby>& InAsyncOp, TDefaultErrorResultInternal<TSharedRef<FLobbyDataEOS>>&& Result)
 	{
-		// Step 3: Store the lobby data on the async op properties.
-
 		const FCreateLobby::Params& Params = InAsyncOp.GetParams();
+		const FString& LobbyIdString = GetOpDataChecked<FString>(InAsyncOp, LobbyIdStringKeyName);
 
 		if (Result.IsError())
 		{
-			// Todo: destroy lobby.
-			// Change DestroyLobbyImpl to take an EOS_LobbyId instead of an FLobbyDataEOS.
+			UE_LOG(LogTemp, Warning, TEXT("[FLobbiesEOSGS::CreateLobby] FLobbyDataRegistryEOS::FindOrCreateFromLobbyDetails Failed: User[%s], Lobby[%s], Result[%s]"),
+				*ToLogString(Params.LocalAccountId), *LobbyIdString, *Result.GetErrorValue().GetLogString());
 
-			// Todo: errors
-			InAsyncOp.SetError(Errors::Unknown(MoveTemp(Result.GetErrorValue())));
-			return;
+			return DestroyLobbyDuringCreate(InAsyncOp, Params.LocalAccountId, LobbyIdString, MoveTemp(Result.GetErrorValue()));
+		}
+		else
+		{
+			// Store lobby data on the operation.
+			InAsyncOp.Data.Set<TSharedRef<FLobbyDataEOS>>(LobbyDataKeyName, MoveTemp(Result.GetOkValue()));
+			return MakeFulfilledPromise<void>().GetFuture();
 		}
-
-		// Store lobby data on the operation.
-		TSharedRef<FLobbyDataEOS> LobbyData = Result.GetOkValue();
-		InAsyncOp.Data.Set<TSharedRef<FLobbyDataEOS>>(LobbyDataKeyName, MoveTemp(LobbyData));
 	})
+	// Step 6: Set attributes for the lobby creator.
 	.Then([this](TOnlineAsyncOp<FCreateLobby>& InAsyncOp)
 	{
-		// Step 4: Set attributes for the lobby creator.
-
 		const FCreateLobby::Params& Params = InAsyncOp.GetParams();
 		const TSharedRef<FLobbyDataEOS>& LobbyData = GetOpDataChecked<TSharedRef<FLobbyDataEOS>>(InAsyncOp, LobbyDataKeyName);
 
-		// Add creator attributes if set.
-		// Todo: make this nicer.
-		const FJoinLobbyLocalUserData* UserData = Params.LocalUsers.FindByPredicate(
-		[LocalAccountId = Params.LocalAccountId](const FJoinLobbyLocalUserData& Data)
-		{
-			return LocalAccountId == Data.LocalAccountId;
-		});
-
-		TSharedRef<FClientLobbyMemberDataChanges> LobbyOwnerAttributes = MakeShared<FClientLobbyMemberDataChanges>();
-
 		// Add owner attributes to operation data.
-		// CreatingMemberData is used to update the local lobby data and for dispatching notifications once creation has completed.
+		// CreatingMemberData is reused to update the local lobby data and for dispatching notifications once creation has completed.
+		TSharedRef<FClientLobbyMemberDataChanges> LobbyOwnerAttributes = MakeShared<FClientLobbyMemberDataChanges>(FClientLobbyMemberDataChanges{ Params.UserAttributes, {} });
 		TMap<FAccountId, TSharedRef<FClientLobbyMemberDataChanges>> CreatingMemberData;
-		CreatingMemberData.Reserve(Params.LocalUsers.Num());
 		CreatingMemberData.Add(Params.LocalAccountId, LobbyOwnerAttributes);
 		InAsyncOp.Data.Set<TMap<FAccountId, TSharedRef<FClientLobbyMemberDataChanges>>>(LobbyMemberChangesKeyName, MoveTemp(CreatingMemberData));
 
-		if (UserData)
+		// Write user attributes to the lobby.
+		if (!Params.UserAttributes.IsEmpty())
 		{
-			LobbyOwnerAttributes->MutatedAttributes = UserData->Attributes;
-
 			FLobbiesModifyLobbyMemberDataImpl::Params ModifyLobbyMemberDataParams;
 			ModifyLobbyMemberDataParams.LobbyData = LobbyData;
 			ModifyLobbyMemberDataParams.LocalAccountId = Params.LocalAccountId;
 			ModifyLobbyMemberDataParams.Changes = LobbyOwnerAttributes;
-
 			return ModifyLobbyMemberDataImpl(MoveTemp(ModifyLobbyMemberDataParams));
 		}
 		else
 		{
-			return MakeFulfilledPromise<TDefaultErrorResult<FLobbiesModifyLobbyMemberDataImpl>>().GetFuture();
+			// Skip attribute write when no user attributes are present.
+			return MakeFulfilledPromise<TDefaultErrorResult<FLobbiesModifyLobbyMemberDataImpl>>(FLobbiesModifyLobbyMemberDataImpl::Result{}).GetFuture();
 		}
 	})
-	.Then([this, DestroyLobbyDuringCreate](TOnlineAsyncOp<FCreateLobby>& InAsyncOp, TDefaultErrorResult<FLobbiesModifyLobbyMemberDataImpl>&& ModifyLobbyOwnerResult)
+	// Step 7: Handle result
+	.Then([this, DestroyLobbyDuringCreate](TOnlineAsyncOp<FCreateLobby>& InAsyncOp, TDefaultErrorResult<FLobbiesModifyLobbyMemberDataImpl>&& Result)
 	{
-		// Step 5: Handle result
-
 		const FCreateLobby::Params& Params = InAsyncOp.GetParams();
 		TSharedRef<FLobbyDataEOS> LobbyData = GetOpDataChecked<TSharedRef<FLobbyDataEOS>>(InAsyncOp, LobbyDataKeyName);
 
-		if (ModifyLobbyOwnerResult.IsError())
+		if (Result.IsError())
 		{
-			return DestroyLobbyDuringCreate(InAsyncOp, LobbyData, Params.LocalAccountId, MoveTemp(ModifyLobbyOwnerResult.GetErrorValue()));
+			UE_LOG(LogTemp, Warning, TEXT("[FLobbiesEOSGS::CreateLobby] ModifyLobbyMemberDataImpl Failed: User[%s], Lobby[%s], Result[%s]"),
+				*ToLogString(Params.LocalAccountId), *LobbyData->GetLobbyIdString(), *Result.GetErrorValue().GetLogString());
+
+			return DestroyLobbyDuringCreate(InAsyncOp, Params.LocalAccountId, LobbyData->GetLobbyIdString(), MoveTemp(Result.GetErrorValue()));
 		}
 		else
 		{
 			return MakeFulfilledPromise<void>().GetFuture();
 		}
 	})
+	// Step 8: Add lobby attributes and set privacy to user setting.
 	.Then([this](TOnlineAsyncOp<FCreateLobby>& InAsyncOp)
 	{
-		// Step 6: Todo: Add other local members.
-
-		// Store member attributes on the operation.
-	})
-	.Then([this](TOnlineAsyncOp<FCreateLobby>& InAsyncOp)
-	{
-		// Step 7: Add lobby attributes, set lobby join policy to user setting.
-
 		const FCreateLobby::Params& Params = InAsyncOp.GetParams();
 		const TSharedRef<FLobbyDataEOS>& LobbyData = GetOpDataChecked<TSharedRef<FLobbyDataEOS>>(InAsyncOp, LobbyDataKeyName);
 
@@ -240,7 +257,7 @@ TOnlineAsyncOpHandle<FCreateLobby> FLobbiesEOSGS::CreateLobby(FCreateLobby::Para
 		TSharedRef<FClientLobbyDataChanges> LobbyChanges = MakeShared<FClientLobbyDataChanges>();
 		InAsyncOp.Data.Set<TSharedRef<FClientLobbyDataChanges>>(LobbyChangesKeyName, LobbyChanges);
 		LobbyChanges->MutatedAttributes = Params.Attributes;
-		LobbyChanges->JoinPolicy = Params.JoinPolicy;
+		LobbyChanges->JoinPolicy.Emplace(Params.JoinPolicy);
 
 		// Add lobby attributes.
 		// Set lobby privacy setting to the user provided value.
@@ -251,47 +268,44 @@ TOnlineAsyncOpHandle<FCreateLobby> FLobbiesEOSGS::CreateLobby(FCreateLobby::Para
 
 		return ModifyLobbyDataImpl(MoveTemp(ModifyLobbyDataParams));
 	})
-	.Then([this, DestroyLobbyDuringCreate](TOnlineAsyncOp<FCreateLobby>& InAsyncOp, TDefaultErrorResult<FLobbiesModifyLobbyDataImpl>&& ModifyLobbyResult) mutable
+	// Step 9: Handle result.
+	.Then([this, DestroyLobbyDuringCreate](TOnlineAsyncOp<FCreateLobby>& InAsyncOp, TDefaultErrorResult<FLobbiesModifyLobbyDataImpl>&& Result) mutable
 	{
-		// Step 8: Handle result.
-
 		const FCreateLobby::Params& Params = InAsyncOp.GetParams();
 		const TSharedRef<FLobbyDataEOS>& LobbyData = GetOpDataChecked<TSharedRef<FLobbyDataEOS>>(InAsyncOp, LobbyDataKeyName);
 
-		if (ModifyLobbyResult.IsError())
+		if (Result.IsError())
 		{
-			return DestroyLobbyDuringCreate(InAsyncOp, LobbyData, Params.LocalAccountId, MoveTemp(ModifyLobbyResult.GetErrorValue()));
+			UE_LOG(LogTemp, Warning, TEXT("[FLobbiesEOSGS::CreateLobby] ModifyLobbyDataImpl Failed: User[%s], Lobby[%s], Result[%s]"),
+				*ToLogString(Params.LocalAccountId), *LobbyData->GetLobbyIdString(), *Result.GetErrorValue().GetLogString());
+
+			return DestroyLobbyDuringCreate(InAsyncOp, Params.LocalAccountId, LobbyData->GetLobbyIdString(), MoveTemp(Result.GetErrorValue()));
 		}
 		else
 		{
 			return MakeFulfilledPromise<void>().GetFuture();
 		}
 	})
+	// Step 10: Add the lobby to the active list, apply changes to the cached lobby object, and signal notifications.
 	.Then([this](TOnlineAsyncOp<FCreateLobby>& InAsyncOp)
 	{
-		// Step 9: Add the lobby to the active list for each member, apply changes to the cached
-		// lobby object, and signal notifications.
-		// The active lobbies list holds a reference to the lobby data to keep it from being cleaned up.
-
 		const FCreateLobby::Params& Params = InAsyncOp.GetParams();
 		const TSharedRef<FLobbyDataEOS>& LobbyData = GetOpDataChecked<TSharedRef<FLobbyDataEOS>>(InAsyncOp, LobbyDataKeyName);
 
-		// Mark the lobby active for each member.
-		for (const FJoinLobbyLocalUserData& LocalUser : Params.LocalUsers)
-		{
-			AddActiveLobby(LocalUser.LocalAccountId, LobbyData);
-		}
+		// Mark the lobby active.
+		AddActiveLobby(Params.LocalAccountId, LobbyData);
 
 		// Add member changes to the lobby changes object.
 		const TSharedRef<FClientLobbyDataChanges>& LobbyChanges = GetOpDataChecked<TSharedRef<FClientLobbyDataChanges>>(InAsyncOp, LobbyChangesKeyName);
-		const TMap<FAccountId, TSharedRef<FClientLobbyMemberDataChanges>>& MemberChanges =
-			GetOpDataChecked<TMap<FAccountId, TSharedRef<FClientLobbyMemberDataChanges>>>(InAsyncOp, LobbyMemberChangesKeyName);
+		const TMap<FAccountId, TSharedRef<FClientLobbyMemberDataChanges>>& MemberChanges = GetOpDataChecked<TMap<FAccountId, TSharedRef<FClientLobbyMemberDataChanges>>>(InAsyncOp, LobbyMemberChangesKeyName);
 		LobbyChanges->MutatedMembers = MemberChanges;
 		LobbyChanges->LocalName = InAsyncOp.GetParams().LocalName;
 
 		// Make local changes to lobby data and generate notifications.
 		LobbyData->GetClientLobbyData()->ApplyLobbyUpdateFromLocalChanges(MoveTemp(*LobbyChanges), LobbyEvents);
 
+		UE_LOG(LogTemp, Log, TEXT("[FLobbiesEOSGS::CreateLobby] Succeeded: User[%s], Lobby[%s]"),
+			*ToLogString(Params.LocalAccountId), *LobbyData->GetLobbyIdString());
 		InAsyncOp.SetResult(FCreateLobby::Result{LobbyData->GetClientLobbyData()->GetPublicDataPtr()});
 	})
 	.Enqueue(GetSerialQueue());
@@ -304,8 +318,10 @@ TOnlineAsyncOpHandle<FFindLobbies> FLobbiesEOSGS::FindLobbies(FFindLobbies::Para
 	TOnlineAsyncOpRef<FFindLobbies> Op = GetOp<FFindLobbies>(MoveTemp(InParams));
 	const FFindLobbies::Params& Params = Op->GetParams();
 
-	if (!Services.Get<FAuthEOSGS>()->IsLoggedIn(Params.LocalAccountId))
+	// Check prerequisites.
+	if (!Services.GetAuthInterface()->IsLoggedIn(Params.LocalAccountId))
 	{
+		UE_LOG(LogTemp, Warning, TEXT("[FLobbiesEOSGS::FindLobbies] Failed: User is not logged in. User[%s]"), *ToLogString(Params.LocalAccountId));
 		Op->SetError(Errors::InvalidUser());
 		return Op->GetHandle();
 	}
@@ -313,6 +329,7 @@ TOnlineAsyncOpHandle<FFindLobbies> FLobbiesEOSGS::FindLobbies(FFindLobbies::Para
 	// Invalidate previous search results.
 	ActiveSearchResults.Remove(Params.LocalAccountId);
 
+	// Start operation.
 	Op->Then([this](TOnlineAsyncOp<FFindLobbies>& InAsyncOp)
 	{
 		return FLobbySearchEOS::Create(LobbyPrerequisites.ToSharedRef(), LobbyDataRegistry.ToSharedRef(), InAsyncOp.GetParams());
@@ -321,15 +338,28 @@ TOnlineAsyncOpHandle<FFindLobbies> FLobbiesEOSGS::FindLobbies(FFindLobbies::Para
 	{
 		if (Result.IsError())
 		{
-			// TODO: Error codes
-			InAsyncOp.SetError(Errors::Unknown(Result.GetErrorValue()));
+			const FFindLobbies::Params& Params = InAsyncOp.GetParams();
+
+			UE_LOG(LogTemp, Warning, TEXT("[FLobbiesEOSGS::JoinLobby] FLobbySearchEOS::Create Failed: User[%s], Result[%s]"),
+				*ToLogString(Params.LocalAccountId), *Result.GetErrorValue().GetLogString());
+			InAsyncOp.SetError(MoveTemp(Result.GetErrorValue()));
 		}
 		else
 		{
-			ActiveSearchResults.Add(InAsyncOp.GetParams().LocalAccountId, Result.GetOkValue());
-			InAsyncOp.SetResult(FFindLobbies::Result{Result.GetOkValue()->GetLobbyResults()});
+			InAsyncOp.Data.Set<TSharedRef<FLobbySearchEOS>>(LobbySearchKeyName, Result.GetOkValue());
 		}
 	})
+	.Then([this](TOnlineAsyncOp<FFindLobbies>& InAsyncOp)
+	{
+		const FFindLobbies::Params& Params = InAsyncOp.GetParams();
+		const TSharedRef<FLobbySearchEOS>& LobbySearch = GetOpDataChecked<TSharedRef<FLobbySearchEOS>>(InAsyncOp, LobbySearchKeyName);
+
+		ActiveSearchResults.Add(InAsyncOp.GetParams().LocalAccountId, LobbySearch);
+
+		UE_LOG(LogTemp, Log, TEXT("[FLobbiesEOSGS::FindLobbies] Succeeded: User[%s]"),
+			*ToLogString(Params.LocalAccountId));
+		InAsyncOp.SetResult(FFindLobbies::Result{ LobbySearch->GetLobbyResults() });
+	})
 	.Enqueue(GetSerialQueue());
 
 	return Op->GetHandle();
@@ -340,66 +370,157 @@ TOnlineAsyncOpHandle<FJoinLobby> FLobbiesEOSGS::JoinLobby(FJoinLobby::Params&& I
 	TOnlineAsyncOpRef<FJoinLobby> Op = GetOp<FJoinLobby>(MoveTemp(InParams));
 	const FJoinLobby::Params& Params = Op->GetParams();
 
-	if (!Services.Get<FAuthEOSGS>()->IsLoggedIn(Params.LocalAccountId))
+	// Check prerequisites.
+	if (!Services.GetAuthInterface()->IsLoggedIn(Params.LocalAccountId))
 	{
+		UE_LOG(LogTemp, Warning, TEXT("[FLobbiesEOSGS::JoinLobby] Failed: User is not logged in. User[%s]"), *ToLogString(Params.LocalAccountId));
 		Op->SetError(Errors::InvalidUser());
 		return Op->GetHandle();
 	}
 
+	if (!Params.LobbyId.IsValid())
+	{
+		UE_LOG(LogTemp, Warning, TEXT("[FLobbiesEOSGS::JoinLobby] Failed: Lobby id is invalid."));
+		Op->SetError(Errors::InvalidParams());
+		return Op->GetHandle();
+	}
+
+	// Lobby data must already exist through local creation, invitation, search, or UI event when joining.
 	TSharedPtr<FLobbyDataEOS> LobbyData = LobbyDataRegistry->Find(Params.LobbyId);
 	if (!LobbyData)
 	{
-		Op->SetError(Errors::Unknown());
+		UE_LOG(LogTemp, Warning, TEXT("[FLobbiesEOSGS::JoinLobby] Failed: Unable to find lobby data. LobbyId[%s]"), *ToLogString(Params.LobbyId));
+		Op->SetError(Errors::InvalidParams());
 		return Op->GetHandle();
 	}
+	Op->Data.Set<TSharedRef<FLobbyDataEOS>>(LobbyDataKeyName, LobbyData.ToSharedRef());
 
-	Op->Then([this, LobbyData](TOnlineAsyncOp<FJoinLobby>& InAsyncOp)
+	// Setup lobby details - Prefer UI event before invitation before search result.
+	TSharedPtr<FLobbyDetailsEOS> LobbyDetails = LobbyData->GetUserLobbyDetails(Params.LocalAccountId);
+	if (!LobbyDetails)
 	{
-		// Join all users to the lobby.
+		UE_LOG(LogTemp, Warning, TEXT("[FLobbiesEOSGS::JoinLobby] Failed: Unable to find lobby details for user. User[%s], Lobby[%s]"),
+			*ToLogString(Params.LocalAccountId), *LobbyData->GetLobbyIdString());
+		Op->SetError(Errors::InvalidParams());
+		return Op->GetHandle();
+	}
+	Op->Data.Set<TSharedRef<FLobbyDetailsEOS>>(LobbyDetailsKeyName, LobbyDetails.ToSharedRef());
 
+	// Check that lobby is compatible with the running game.
+	if (!LobbyDetails->IsBucketCompatible())
+	{
+		UE_LOG(LogTemp, Warning, TEXT("[FLobbiesEOSGS::JoinLobby] Failed: Lobby is not compatible with the application version. Lobby[%s], Bucket[%s], ExpectedBucket[%s]."),
+			*LobbyData->GetLobbyIdString(), *ToLogString(LobbyDetails->GetInfo()->GetBucketId()), *ToLogString(LobbyPrerequisites->BucketId));
+		Op->SetError(Errors::IncompatibleVersion());
+		return Op->GetHandle();
+	}
+
+	// Start operation.
+	// Step 1. Join the lobby.
+	Op->Then([this](TOnlineAsyncOp<FJoinLobby>& InAsyncOp, TPromise<const EOS_Lobby_JoinLobbyCallbackInfo*>&& Promise)
+	{
 		const FJoinLobby::Params& Params = InAsyncOp.GetParams();
+		const TSharedRef<FLobbyDetailsEOS>& LobbyDetails = GetOpDataChecked<TSharedRef<FLobbyDetailsEOS>>(InAsyncOp, LobbyDetailsKeyName);
+
+		EOS_Lobby_JoinLobbyOptions JoinLobbyOptions = {};
+		JoinLobbyOptions.ApiVersion = EOS_LOBBY_JOINLOBBY_API_LATEST;
+		JoinLobbyOptions.LobbyDetailsHandle = LobbyDetails->GetEOSHandle();
+		JoinLobbyOptions.LocalUserId = GetProductUserIdChecked(Params.LocalAccountId);
+		JoinLobbyOptions.bPresenceEnabled = false; // todo
+		JoinLobbyOptions.LocalRTCOptions = nullptr;
+		static_assert(EOS_LOBBY_JOINLOBBY_API_LATEST == 3, "EOS_Lobby_JoinLobbyOptions updated, check new fields");
 
-		FLobbiesJoinLobbyImpl::Params JoinParams;
-		JoinParams.LobbyData = LobbyData;
-		JoinParams.LocalAccountId = Params.LocalAccountId;
-		JoinParams.LocalName = Params.LocalName;
-		JoinParams.LocalUsers = Params.LocalUsers;
-		return JoinLobbyImpl(MoveTemp(JoinParams));
+		EOS_Async(EOS_Lobby_JoinLobby, LobbyPrerequisites->LobbyInterfaceHandle, JoinLobbyOptions, MoveTemp(Promise));
 	})
-	.Then([this](TOnlineAsyncOp<FJoinLobby>& InAsyncOp, TDefaultErrorResult<FLobbiesJoinLobbyImpl> Result)
+	// Step 2. Handle join result.
+	.Then([this](TOnlineAsyncOp<FJoinLobby>& InAsyncOp, const EOS_Lobby_JoinLobbyCallbackInfo* Data)
 	{
-		// Handle result.
-		if (Result.IsError())
+		if (Data->ResultCode != EOS_EResult::EOS_Success)
 		{
-			// TODO: Error codes
-			InAsyncOp.SetError(Errors::Unknown(Result.GetErrorValue()));
+			const FJoinLobby::Params& Params = InAsyncOp.GetParams();
+			const TSharedRef<FLobbyDataEOS>& LobbyData = GetOpDataChecked<TSharedRef<FLobbyDataEOS>>(InAsyncOp, LobbyDataKeyName);
+
+			UE_LOG(LogTemp, Warning, TEXT("[FLobbiesEOSGS::JoinLobby] EOS_Lobby_JoinLobby Failed: User[%s], Lobby[%s], Result[%s]"),
+				*ToLogString(Params.LocalAccountId), *LobbyData->GetLobbyIdString(), *LexToString(Data->ResultCode));
+			InAsyncOp.SetError(Errors::FromEOSResult(Data->ResultCode));
 		}
 	})
-	.Then([this, LobbyData](TOnlineAsyncOp<FJoinLobby>& InAsyncOp)
+	// Step 3. Set member attributes.
+	.Then([this](TOnlineAsyncOp<FJoinLobby>& InAsyncOp)
 	{
-		// Mark the lobby active for each member.
-		// Add users to local lobby data and dispatch notifications.
-
 		const FJoinLobby::Params& Params = InAsyncOp.GetParams();
+		const TSharedRef<FLobbyDataEOS>& LobbyData = GetOpDataChecked<TSharedRef<FLobbyDataEOS>>(InAsyncOp, LobbyDataKeyName);
 
-		// Mark the lobby active for each member.
-		for (const FJoinLobbyLocalUserData& LocalUser : Params.LocalUsers)
+		FLobbiesModifyLobbyMemberDataImpl::Params ModifyParams;
+		ModifyParams.LobbyData = LobbyData;
+		ModifyParams.LocalAccountId = Params.LocalAccountId;
+		ModifyParams.Changes = MakeShared<FClientLobbyMemberDataChanges>();
+		ModifyParams.Changes->MutatedAttributes = Params.UserAttributes;
+		return ModifyLobbyMemberDataImpl(MoveTemp(ModifyParams));
+	})
+	// Step 4. Handle result.
+	.Then([this](TOnlineAsyncOp<FJoinLobby>& InAsyncOp, TDefaultErrorResult<FLobbiesModifyLobbyMemberDataImpl>&& Result)
+	{
+		if (Result.IsError())
 		{
-			AddActiveLobby(LocalUser.LocalAccountId, LobbyData.ToSharedRef());
-		}
+			const FJoinLobby::Params& Params = InAsyncOp.GetParams();
+			const TSharedRef<FLobbyDataEOS>& LobbyData = GetOpDataChecked<TSharedRef<FLobbyDataEOS>>(InAsyncOp, LobbyDataKeyName);
 
-		// todo: figure out a butter way to handle attribute parameters.
-		FClientLobbyDataChanges LobbyChanges;
-		LobbyChanges.LocalName = InAsyncOp.GetParams().LocalName;
+			UE_LOG(LogTemp, Warning, TEXT("[FLobbiesEOSGS::JoinLobby] ModifyLobbyMemberDataImpl Failed: User[%s], Lobby[%s], Result[%s]"),
+				*ToLogString(Params.LocalAccountId), *LobbyData->GetLobbyIdString(), *Result.GetErrorValue().GetLogString());
 
-		for (const FJoinLobbyLocalUserData& LocalUser : Params.LocalUsers)
+			// Failed to set attributes - leave the lobby.
+			FLobbiesLeaveLobbyImpl::Params LeaveLobbyParams;
+			LeaveLobbyParams.LobbyData = LobbyData;
+			LeaveLobbyParams.LocalAccountId = Params.LocalAccountId;
+
+			TPromise<void> Promise;
+			TFuture<void> Future = Promise.GetFuture();
+
+			LeaveLobbyImpl(MoveTemp(LeaveLobbyParams))
+			.Next(
+			[
+				LocalAccountId = Params.LocalAccountId,
+				LobbyData,
+				InAsyncOp = InAsyncOp.AsShared(),
+				ErrorResult = MoveTemp(Result.GetErrorValue()),
+				Promise = MoveTemp(Promise)
+			](TDefaultErrorResult<FLobbiesLeaveLobbyImpl>&& Result) mutable
+			{
+				if (Result.IsError())
+				{
+					UE_LOG(LogTemp, Warning, TEXT("[FLobbiesEOSGS::JoinLobby] LeaveLobbyImpl Failed: User[%s], Lobby[%s], Result[%s]"),
+						*ToLogString(LocalAccountId), *LobbyData->GetLobbyIdString(), *Result.GetErrorValue().GetLogString());
+				}
+
+				InAsyncOp->SetError(MoveTemp(ErrorResult));
+				Promise.EmplaceValue();
+			});
+
+			return Future;
+		}
+		else
 		{
-			TSharedRef<FClientLobbyMemberDataChanges> MemberDataChanges = MakeShared<FClientLobbyMemberDataChanges>();
-			MemberDataChanges->MutatedAttributes = LocalUser.Attributes;
-			LobbyChanges.MutatedMembers.Add(LocalUser.LocalAccountId, MoveTemp(MemberDataChanges));
+			return MakeFulfilledPromise<void>().GetFuture();
 		}
+	})
+	// Step 5. Bookkeeping and notifications.
+	.Then([this](TOnlineAsyncOp<FJoinLobby>& InAsyncOp)
+	{
+		const FJoinLobby::Params& Params = InAsyncOp.GetParams();
+		const TSharedRef<FLobbyDataEOS>& LobbyData = GetOpDataChecked<TSharedRef<FLobbyDataEOS>>(InAsyncOp, LobbyDataKeyName);
+
+		// Mark the lobby active.
+		AddActiveLobby(Params.LocalAccountId, LobbyData);
+
+		// Apply local changes and dispatch notifications.
+		FClientLobbyDataChanges LobbyChanges;
+		LobbyChanges.LocalName = InAsyncOp.GetParams().LocalName;
+		LobbyChanges.MutatedMembers.Add(Params.LocalAccountId, MakeShared<FClientLobbyMemberDataChanges>(FClientLobbyMemberDataChanges{ Params.UserAttributes, {} }));
 		LobbyData->GetClientLobbyData()->ApplyLobbyUpdateFromLocalChanges(MoveTemp(LobbyChanges), LobbyEvents);
 
+		UE_LOG(LogTemp, Log, TEXT("[FLobbiesEOSGS::JoinLobby] Succeeded: User[%s], Lobby[%s]"),
+			*ToLogString(Params.LocalAccountId), *LobbyData->GetLobbyIdString());
 		InAsyncOp.SetResult(FJoinLobby::Result{ LobbyData->GetClientLobbyData()->GetPublicDataPtr()});
 	})
 	.Enqueue(GetSerialQueue());
@@ -412,33 +533,59 @@ TOnlineAsyncOpHandle<FLeaveLobby> FLobbiesEOSGS::LeaveLobby(FLeaveLobby::Params&
 	TOnlineAsyncOpRef<FLeaveLobby> Op = GetOp<FLeaveLobby>(MoveTemp(InParams));
 	const FLeaveLobby::Params& Params = Op->GetParams();
 
-	if (!Services.Get<FAuthEOSGS>()->IsLoggedIn(Params.LocalAccountId))
+	// Check prerequisites.
+	if (!Services.GetAuthInterface()->IsLoggedIn(Params.LocalAccountId))
 	{
+		UE_LOG(LogTemp, Warning, TEXT("[FLobbiesEOSGS::LeaveLobby] Failed: User is not logged in. User[%s]"), *ToLogString(Params.LocalAccountId));
 		Op->SetError(Errors::InvalidUser());
 		return Op->GetHandle();
 	}
 
+	if (!Params.LobbyId.IsValid())
+	{
+		UE_LOG(LogTemp, Warning, TEXT("[FLobbiesEOSGS::LeaveLobby] Failed: Lobby id is invalid."));
+		Op->SetError(Errors::InvalidParams());
+		return Op->GetHandle();
+	}
+
 	TSharedPtr<FLobbyDataEOS> LobbyData = LobbyDataRegistry->Find(Params.LobbyId);
 	if (!LobbyData)
 	{
-		Op->SetError(Errors::Unknown());
+		UE_LOG(LogTemp, Warning, TEXT("[FLobbiesEOSGS::LeaveLobby] Failed: Unable to find lobby data. LobbyId[%s]"), *ToLogString(Params.LobbyId));
+		Op->SetError(Errors::InvalidParams());
 		return Op->GetHandle();
 	}
+	Op->Data.Set<TSharedRef<FLobbyDataEOS>>(LobbyDataKeyName, LobbyData.ToSharedRef());
 
-	Op->Then([this, LobbyData](TOnlineAsyncOp<FLeaveLobby>& InAsyncOp)
+	// Start operation.
+	Op->Then([this](TOnlineAsyncOp<FLeaveLobby>& InAsyncOp)
 	{
-		// Remove the user from the EOS lobby.
-
 		const FLeaveLobby::Params& Params = InAsyncOp.GetParams();
+		const TSharedRef<FLobbyDataEOS>& LobbyData = GetOpDataChecked<TSharedRef<FLobbyDataEOS>>(InAsyncOp, LobbyDataKeyName);
 
 		FLobbiesLeaveLobbyImpl::Params LeaveParams;
 		LeaveParams.LobbyData = LobbyData;
 		LeaveParams.LocalAccountId = Params.LocalAccountId;
 		return LeaveLobbyImpl(MoveTemp(LeaveParams));
 	})
-	.Then([this, LobbyData](TOnlineAsyncOp<FLeaveLobby>& InAsyncOp, TDefaultErrorResult<FLobbiesLeaveLobbyImpl> Result)
+	.Then([this](TOnlineAsyncOp<FLeaveLobby>& InAsyncOp, TDefaultErrorResult<FLobbiesLeaveLobbyImpl>&& Result)
+	{
+		if (Result.IsError())
+		{
+			const FLeaveLobby::Params& Params = InAsyncOp.GetParams();
+			const TSharedRef<FLobbyDataEOS>& LobbyData = GetOpDataChecked<TSharedRef<FLobbyDataEOS>>(InAsyncOp, LobbyDataKeyName);
+
+			UE_LOG(LogTemp, Warning, TEXT("[FLobbiesEOSGS::LeaveLobby] LeaveLobbyImpl Failed: User[%s], Lobby[%s], Result[%s]"),
+				*ToLogString(Params.LocalAccountId), *LobbyData->GetLobbyIdString(), *Result.GetErrorValue().GetLogString());
+
+			// When leaving a lobby consume the error and report success.
+			//InAsyncOp.SetError(Result.GetErrorValue());
+		}
+	})
+	.Then([this](TOnlineAsyncOp<FLeaveLobby>& InAsyncOp)
 	{
 		const FLeaveLobby::Params& Params = InAsyncOp.GetParams();
+		const TSharedRef<FLobbyDataEOS>& LobbyData = GetOpDataChecked<TSharedRef<FLobbyDataEOS>>(InAsyncOp, LobbyDataKeyName);
 
 		// Remove the user from the local lobby data and dispatch notifications.
 		FClientLobbyDataChanges LobbyChanges;
@@ -449,18 +596,12 @@ TOnlineAsyncOpHandle<FLeaveLobby> FLobbiesEOSGS::LeaveLobby(FLeaveLobby::Params&
 		// The lobby data will be cleaned up once all references are removed.
 		for (FAccountId LeavingMember : ApplyResult.LeavingLocalMembers)
 		{
-			RemoveActiveLobby(LeavingMember, LobbyData.ToSharedRef());
+			RemoveActiveLobby(LeavingMember, LobbyData);
 		}
 
-		if (Result.IsError())
-		{
-			// TODO: Error codes
-			InAsyncOp.SetError(Errors::Unknown(Result.GetErrorValue()));
-		}
-		else
-		{
-			InAsyncOp.SetResult(FLeaveLobby::Result{});
-		}
+		UE_LOG(LogTemp, Log, TEXT("[FLobbiesEOSGS::LeaveLobby] Succeeded: User[%s], Lobby[%s]"),
+			*ToLogString(Params.LocalAccountId), *LobbyData->GetLobbyIdString());
+		InAsyncOp.SetResult(FLeaveLobby::Result{});
 	})
 	.Enqueue(GetSerialQueue());
 
@@ -472,34 +613,63 @@ TOnlineAsyncOpHandle<FInviteLobbyMember> FLobbiesEOSGS::InviteLobbyMember(FInvit
 	TOnlineAsyncOpRef<FInviteLobbyMember> Op = GetOp<FInviteLobbyMember>(MoveTemp(InParams));
 	const FInviteLobbyMember::Params& Params = Op->GetParams();
 
-	if (!Services.Get<FAuthEOSGS>()->IsLoggedIn(Params.LocalAccountId))
+	// Check prerequisites.
+	if (!Services.GetAuthInterface()->IsLoggedIn(Params.LocalAccountId))
 	{
+		UE_LOG(LogTemp, Warning, TEXT("[FLobbiesEOSGS::InviteLobbyMember] Failed: User is not logged in. User[%s]"), *ToLogString(Params.LocalAccountId));
 		Op->SetError(Errors::InvalidUser());
 		return Op->GetHandle();
 	}
 
+	if (!Params.LobbyId.IsValid())
+	{
+		UE_LOG(LogTemp, Warning, TEXT("[FLobbiesEOSGS::InviteLobbyMember] Failed: Lobby id is invalid."));
+		Op->SetError(Errors::InvalidParams());
+		return Op->GetHandle();
+	}
+
+	TSharedPtr<FLobbyDataEOS> LobbyData = LobbyDataRegistry->Find(Params.LobbyId);
+	if (!LobbyData)
+	{
+		UE_LOG(LogTemp, Warning, TEXT("[FLobbiesEOSGS::InviteLobbyMember] Failed: Unable to find lobby data. LobbyId[%s]"), *ToLogString(Params.LobbyId));
+		Op->SetError(Errors::InvalidParams());
+		return Op->GetHandle();
+	}
+	Op->Data.Set<TSharedRef<FLobbyDataEOS>>(LobbyDataKeyName, LobbyData.ToSharedRef());
+
+	// Start operation.
 	Op->Then([this](TOnlineAsyncOp<FInviteLobbyMember>& InAsyncOp)
 	{
 		const FInviteLobbyMember::Params& Params = InAsyncOp.GetParams();
+		const TSharedRef<FLobbyDataEOS>& LobbyData = GetOpDataChecked<TSharedRef<FLobbyDataEOS>>(InAsyncOp, LobbyDataKeyName);
 
 		FLobbiesInviteLobbyMemberImpl::Params InviteParams;
-		InviteParams.LobbyData = LobbyDataRegistry->Find(Params.LobbyId);
+		InviteParams.LobbyData = LobbyData;
 		InviteParams.LocalAccountId = Params.LocalAccountId;
 		InviteParams.TargetAccountId = Params.TargetAccountId;
 		return InviteLobbyMemberImpl(MoveTemp(InviteParams));
 	})
-	.Then([this](TOnlineAsyncOp<FInviteLobbyMember>& InAsyncOp, TDefaultErrorResult<FLobbiesInviteLobbyMemberImpl> Result)
+	.Then([this](TOnlineAsyncOp<FInviteLobbyMember>& InAsyncOp, TDefaultErrorResult<FLobbiesInviteLobbyMemberImpl>&& Result)
 	{
 		if (Result.IsError())
 		{
-			// TODO: Error codes
-			InAsyncOp.SetError(Errors::Unknown(Result.GetErrorValue()));
-		}
-		else
-		{
-			InAsyncOp.SetResult(FInviteLobbyMember::Result{});
+			const FInviteLobbyMember::Params& Params = InAsyncOp.GetParams();
+			const TSharedRef<FLobbyDataEOS>& LobbyData = GetOpDataChecked<TSharedRef<FLobbyDataEOS>>(InAsyncOp, LobbyDataKeyName);
+
+			UE_LOG(LogTemp, Warning, TEXT("[FLobbiesEOSGS::InviteLobbyMember] InviteLobbyMemberImpl Failed: User[%s], Lobby[%s], Result[%s]"),
+				*ToLogString(Params.LocalAccountId), *LobbyData->GetLobbyIdString(), *Result.GetErrorValue().GetLogString());
+			InAsyncOp.SetError(MoveTemp(Result.GetErrorValue()));
 		}
 	})
+	.Then([this](TOnlineAsyncOp<FInviteLobbyMember>& InAsyncOp)
+	{
+		const FInviteLobbyMember::Params& Params = InAsyncOp.GetParams();
+		const TSharedRef<FLobbyDataEOS>& LobbyData = GetOpDataChecked<TSharedRef<FLobbyDataEOS>>(InAsyncOp, LobbyDataKeyName);
+
+		UE_LOG(LogTemp, Verbose, TEXT("[FLobbiesEOSGS::InviteLobbyMember] Succeeded: User[%s], Lobby[%s]"),
+			*ToLogString(Params.LocalAccountId), *LobbyData->GetLobbyIdString());
+		InAsyncOp.SetResult(FInviteLobbyMember::Result{});
+	})
 	.Enqueue(GetSerialQueue());
 
 	return Op->GetHandle();
@@ -510,12 +680,22 @@ TOnlineAsyncOpHandle<FDeclineLobbyInvitation> FLobbiesEOSGS::DeclineLobbyInvitat
 	TOnlineAsyncOpRef<FDeclineLobbyInvitation> Op = GetOp<FDeclineLobbyInvitation>(MoveTemp(InParams));
 	const FDeclineLobbyInvitation::Params& Params = Op->GetParams();
 
-	if (!Services.Get<FAuthEOSGS>()->IsLoggedIn(Params.LocalAccountId))
+	// Check prerequisites.
+	if (!Services.GetAuthInterface()->IsLoggedIn(Params.LocalAccountId))
 	{
+		UE_LOG(LogTemp, Warning, TEXT("[FLobbiesEOSGS::DeclineLobbyInvitation] Failed: User is not logged in. User[%s]"), *ToLogString(Params.LocalAccountId));
 		Op->SetError(Errors::InvalidUser());
 		return Op->GetHandle();
 	}
 
+	if (!Params.LobbyId.IsValid())
+	{
+		UE_LOG(LogTemp, Warning, TEXT("[FLobbiesEOSGS::DeclineLobbyInvitation] Failed: Lobby id is invalid."));
+		Op->SetError(Errors::InvalidParams());
+		return Op->GetHandle();
+	}
+
+	// Start operation.
 	Op->Then([this](TOnlineAsyncOp<FDeclineLobbyInvitation>& InAsyncOp)
 	{
 		const FDeclineLobbyInvitation::Params& Params = InAsyncOp.GetParams();
@@ -525,18 +705,27 @@ TOnlineAsyncOpHandle<FDeclineLobbyInvitation> FLobbiesEOSGS::DeclineLobbyInvitat
 		DeclineParams.LobbyId = Params.LobbyId;
 		return DeclineLobbyInvitationImpl(MoveTemp(DeclineParams));
 	})
-	.Then([this](TOnlineAsyncOp<FDeclineLobbyInvitation>& InAsyncOp, TDefaultErrorResult<FLobbiesDeclineLobbyInvitationImpl> Result)
+	.Then([this](TOnlineAsyncOp<FDeclineLobbyInvitation>& InAsyncOp, TDefaultErrorResult<FLobbiesDeclineLobbyInvitationImpl>&& Result)
 	{
 		if (Result.IsError())
 		{
-			// TODO: Error codes
-			InAsyncOp.SetError(Errors::Unknown(Result.GetErrorValue()));
-		}
-		else
-		{
-			InAsyncOp.SetResult(FDeclineLobbyInvitation::Result{});
+			const FDeclineLobbyInvitation::Params& Params = InAsyncOp.GetParams();
+			const TSharedRef<FLobbyDataEOS>& LobbyData = GetOpDataChecked<TSharedRef<FLobbyDataEOS>>(InAsyncOp, LobbyDataKeyName);
+
+			UE_LOG(LogTemp, Warning, TEXT("[FLobbiesEOSGS::DeclineLobbyInvitation] DeclineLobbyInvitationImpl Failed: User[%s], Lobby[%s], Result[%s]"),
+				*ToLogString(Params.LocalAccountId), *LobbyData->GetLobbyIdString(), *Result.GetErrorValue().GetLogString());
+			InAsyncOp.SetError(MoveTemp(Result.GetErrorValue()));
 		}
 	})
+	.Then([this](TOnlineAsyncOp<FDeclineLobbyInvitation>& InAsyncOp)
+	{
+		const FDeclineLobbyInvitation::Params& Params = InAsyncOp.GetParams();
+		const TSharedRef<FLobbyDataEOS>& LobbyData = GetOpDataChecked<TSharedRef<FLobbyDataEOS>>(InAsyncOp, LobbyDataKeyName);
+
+		UE_LOG(LogTemp, Verbose, TEXT("[FLobbiesEOSGS::DeclineLobbyInvitation] Succeeded: User[%s], Lobby[%s]"),
+			*ToLogString(Params.LocalAccountId), *LobbyData->GetLobbyIdString());
+		InAsyncOp.SetResult(FDeclineLobbyInvitation::Result{});
+	})
 	.Enqueue(GetSerialQueue());
 
 	return Op->GetHandle();
@@ -547,20 +736,31 @@ TOnlineAsyncOpHandle<FKickLobbyMember> FLobbiesEOSGS::KickLobbyMember(FKickLobby
 	TOnlineAsyncOpRef<FKickLobbyMember> Op = GetOp<FKickLobbyMember>(MoveTemp(InParams));
 	const FKickLobbyMember::Params& Params = Op->GetParams();
 
-	if (!Services.Get<FAuthEOSGS>()->IsLoggedIn(Params.LocalAccountId))
+	// Check prerequisites.
+	if (!Services.GetAuthInterface()->IsLoggedIn(Params.LocalAccountId))
 	{
+		UE_LOG(LogTemp, Warning, TEXT("[FLobbiesEOSGS::KickLobbyMember] Failed: User is not logged in. User[%s]"), *ToLogString(Params.LocalAccountId));
 		Op->SetError(Errors::InvalidUser());
 		return Op->GetHandle();
 	}
 
+	if (!Params.LobbyId.IsValid())
+	{
+		UE_LOG(LogTemp, Warning, TEXT("[FLobbiesEOSGS::KickLobbyMember] Failed: Lobby id is invalid."));
+		Op->SetError(Errors::InvalidParams());
+		return Op->GetHandle();
+	}
+
 	TSharedPtr<FLobbyDataEOS> LobbyData = LobbyDataRegistry->Find(Params.LobbyId);
 	if (!LobbyData)
 	{
+		UE_LOG(LogTemp, Warning, TEXT("[FLobbiesEOSGS::KickLobbyMember] Failed: Unable to find lobby data. LobbyId[%s]"), *ToLogString(Params.LobbyId));
 		Op->SetError(Errors::InvalidParams());
 		return Op->GetHandle();
 	}
 	Op->Data.Set<TSharedRef<FLobbyDataEOS>>(LobbyDataKeyName, LobbyData.ToSharedRef());
 
+	// Start operation.
 	Op->Then([this](TOnlineAsyncOp<FKickLobbyMember>& InAsyncOp)
 	{
 		const FKickLobbyMember::Params& Params = InAsyncOp.GetParams();
@@ -572,12 +772,16 @@ TOnlineAsyncOpHandle<FKickLobbyMember> FLobbiesEOSGS::KickLobbyMember(FKickLobby
 		KickParams.TargetAccountId = Params.TargetAccountId;
 		return KickLobbyMemberImpl(MoveTemp(KickParams));
 	})
-	.Then([this](TOnlineAsyncOp<FKickLobbyMember>& InAsyncOp, TDefaultErrorResult<FLobbiesKickLobbyMemberImpl> Result)
+	.Then([this](TOnlineAsyncOp<FKickLobbyMember>& InAsyncOp, TDefaultErrorResult<FLobbiesKickLobbyMemberImpl>&& Result)
 	{
 		if (Result.IsError())
 		{
-			// TODO: Error codes
-			InAsyncOp.SetError(Errors::Unknown(Result.GetErrorValue()));
+			const FKickLobbyMember::Params& Params = InAsyncOp.GetParams();
+			const TSharedRef<FLobbyDataEOS>& LobbyData = GetOpDataChecked<TSharedRef<FLobbyDataEOS>>(InAsyncOp, LobbyDataKeyName);
+
+			UE_LOG(LogTemp, Warning, TEXT("[FLobbiesEOSGS::KickLobbyMember] KickLobbyMemberImpl Failed: User[%s], Lobby[%s], Result[%s]"),
+				*ToLogString(Params.LocalAccountId), *LobbyData->GetLobbyIdString(), *Result.GetErrorValue().GetLogString());
+			InAsyncOp.SetError(MoveTemp(Result.GetErrorValue()));
 		}
 	})
 	.Then([this](TOnlineAsyncOp<FKickLobbyMember>& InAsyncOp)
@@ -589,6 +793,9 @@ TOnlineAsyncOpHandle<FKickLobbyMember> FLobbiesEOSGS::KickLobbyMember(FKickLobby
 		FClientLobbyDataChanges LobbyChanges;
 		LobbyChanges.LeavingMembers.Add(Params.TargetAccountId, ELobbyMemberLeaveReason::Kicked);
 		LobbyData->GetClientLobbyData()->ApplyLobbyUpdateFromLocalChanges(MoveTemp(LobbyChanges), LobbyEvents);
+
+		UE_LOG(LogTemp, Verbose, TEXT("[FLobbiesEOSGS::KickLobbyMember] Succeeded: User[%s], Lobby[%s]"),
+			*ToLogString(Params.LocalAccountId), *LobbyData->GetLobbyIdString());
 		InAsyncOp.SetResult(FKickLobbyMember::Result{});
 	})
 	.Enqueue(GetSerialQueue());
@@ -598,23 +805,34 @@ TOnlineAsyncOpHandle<FKickLobbyMember> FLobbiesEOSGS::KickLobbyMember(FKickLobby
 
 TOnlineAsyncOpHandle<FPromoteLobbyMember> FLobbiesEOSGS::PromoteLobbyMember(FPromoteLobbyMember::Params&& InParams)
 {
+	// Check prerequisites.
 	TOnlineAsyncOpRef<FPromoteLobbyMember> Op = GetOp<FPromoteLobbyMember>(MoveTemp(InParams));
 	const FPromoteLobbyMember::Params& Params = Op->GetParams();
 
-	if (!Services.Get<FAuthEOSGS>()->IsLoggedIn(Params.LocalAccountId))
+	if (!Services.GetAuthInterface()->IsLoggedIn(Params.LocalAccountId))
 	{
+		UE_LOG(LogTemp, Warning, TEXT("[FLobbiesEOSGS::FPromoteLobbyMember] Failed: User is not logged in. User[%s]"), *ToLogString(Params.LocalAccountId));
 		Op->SetError(Errors::InvalidUser());
 		return Op->GetHandle();
 	}
 
+	if (!Params.LobbyId.IsValid())
+	{
+		UE_LOG(LogTemp, Warning, TEXT("[FLobbiesEOSGS::FPromoteLobbyMember] Failed: Lobby id is invalid."));
+		Op->SetError(Errors::InvalidParams());
+		return Op->GetHandle();
+	}
+
 	TSharedPtr<FLobbyDataEOS> LobbyData = LobbyDataRegistry->Find(Params.LobbyId);
 	if (!LobbyData)
 	{
+		UE_LOG(LogTemp, Warning, TEXT("[FLobbiesEOSGS::FPromoteLobbyMember] Failed: Unable to find lobby data. LobbyId[%s]"), *ToLogString(Params.LobbyId));
 		Op->SetError(Errors::InvalidParams());
 		return Op->GetHandle();
 	}
 	Op->Data.Set<TSharedRef<FLobbyDataEOS>>(LobbyDataKeyName, LobbyData.ToSharedRef());
 
+	// Start operation.
 	Op->Then([this](TOnlineAsyncOp<FPromoteLobbyMember>& InAsyncOp)
 	{
 		const FPromoteLobbyMember::Params& Params = InAsyncOp.GetParams();
@@ -626,12 +844,16 @@ TOnlineAsyncOpHandle<FPromoteLobbyMember> FLobbiesEOSGS::PromoteLobbyMember(FPro
 		PromoteParams.TargetAccountId = Params.TargetAccountId;
 		return PromoteLobbyMemberImpl(MoveTemp(PromoteParams));
 	})
-	.Then([this](TOnlineAsyncOp<FPromoteLobbyMember>& InAsyncOp, TDefaultErrorResult<FLobbiesPromoteLobbyMemberImpl> Result)
+	.Then([this](TOnlineAsyncOp<FPromoteLobbyMember>& InAsyncOp, TDefaultErrorResult<FLobbiesPromoteLobbyMemberImpl>&& Result)
 	{
 		if (Result.IsError())
 		{
-			// TODO: Error codes
-			InAsyncOp.SetError(Errors::Unknown(Result.GetErrorValue()));
+			const FPromoteLobbyMember::Params& Params = InAsyncOp.GetParams();
+			const TSharedRef<FLobbyDataEOS>& LobbyData = GetOpDataChecked<TSharedRef<FLobbyDataEOS>>(InAsyncOp, LobbyDataKeyName);
+
+			UE_LOG(LogTemp, Warning, TEXT("[FLobbiesEOSGS::FPromoteLobbyMember] PromoteLobbyMemberImpl Failed: User[%s], Lobby[%s], Result[%s]"),
+				*ToLogString(Params.LocalAccountId), *LobbyData->GetLobbyIdString(), *Result.GetErrorValue().GetLogString());
+			InAsyncOp.SetError(MoveTemp(Result.GetErrorValue()));
 		}
 	})
 	.Then([this](TOnlineAsyncOp<FPromoteLobbyMember>& InAsyncOp)
@@ -643,6 +865,9 @@ TOnlineAsyncOpHandle<FPromoteLobbyMember> FLobbiesEOSGS::PromoteLobbyMember(FPro
 		FClientLobbyDataChanges LobbyChanges;
 		LobbyChanges.OwnerAccountId = Params.TargetAccountId;
 		LobbyData->GetClientLobbyData()->ApplyLobbyUpdateFromLocalChanges(MoveTemp(LobbyChanges), LobbyEvents);
+
+		UE_LOG(LogTemp, Verbose, TEXT("[FLobbiesEOSGS::FPromoteLobbyMember] Succeeded: User[%s], Lobby[%s]"),
+			*ToLogString(Params.LocalAccountId), *LobbyData->GetLobbyIdString());
 		InAsyncOp.SetResult(FPromoteLobbyMember::Result{});
 	})
 	.Enqueue(GetSerialQueue());
@@ -655,20 +880,31 @@ TOnlineAsyncOpHandle<FModifyLobbyJoinPolicy> FLobbiesEOSGS::ModifyLobbyJoinPolic
 	TOnlineAsyncOpRef<FModifyLobbyJoinPolicy> Op = GetOp<FModifyLobbyJoinPolicy>(MoveTemp(InParams));
 	const FModifyLobbyJoinPolicy::Params& Params = Op->GetParams();
 
-	if (!Services.Get<FAuthEOSGS>()->IsLoggedIn(Params.LocalAccountId))
+	// Check prerequisites.
+	if (!Services.GetAuthInterface()->IsLoggedIn(Params.LocalAccountId))
 	{
+		UE_LOG(LogTemp, Warning, TEXT("[FLobbiesEOSGS::ModifyLobbyJoinPolicy] Failed: User is not logged in. User[%s]"), *ToLogString(Params.LocalAccountId));
 		Op->SetError(Errors::InvalidUser());
 		return Op->GetHandle();
 	}
 
+	if (!Params.LobbyId.IsValid())
+	{
+		UE_LOG(LogTemp, Warning, TEXT("[FLobbiesEOSGS::ModifyLobbyJoinPolicy] Failed: Lobby id is invalid."));
+		Op->SetError(Errors::InvalidParams());
+		return Op->GetHandle();
+	}
+
 	TSharedPtr<FLobbyDataEOS> LobbyData = LobbyDataRegistry->Find(Params.LobbyId);
 	if (!LobbyData)
 	{
+		UE_LOG(LogTemp, Warning, TEXT("[FLobbiesEOSGS::ModifyLobbyJoinPolicy] Failed: Unable to find lobby data. LobbyId[%s]"), *ToLogString(Params.LobbyId));
 		Op->SetError(Errors::InvalidParams());
 		return Op->GetHandle();
 	}
 	Op->Data.Set<TSharedRef<FLobbyDataEOS>>(LobbyDataKeyName, LobbyData.ToSharedRef());
 
+	// Start operation.
 	Op->Then([this](TOnlineAsyncOp<FModifyLobbyJoinPolicy>& InAsyncOp)
 	{
 		const FModifyLobbyJoinPolicy::Params& Params = InAsyncOp.GetParams();
@@ -684,21 +920,29 @@ TOnlineAsyncOpHandle<FModifyLobbyJoinPolicy> FLobbiesEOSGS::ModifyLobbyJoinPolic
 		ModifyLobbyDataParams.Changes = LobbyChanges;
 		return ModifyLobbyDataImpl(MoveTemp(ModifyLobbyDataParams));
 	})
-	.Then([this](TOnlineAsyncOp<FModifyLobbyJoinPolicy>& InAsyncOp, TDefaultErrorResult<FLobbiesModifyLobbyDataImpl> Result)
+	.Then([this](TOnlineAsyncOp<FModifyLobbyJoinPolicy>& InAsyncOp, TDefaultErrorResult<FLobbiesModifyLobbyDataImpl>&& Result)
 	{
 		if (Result.IsError())
 		{
-			// TODO: Error codes
-			InAsyncOp.SetError(Errors::Unknown(Result.GetErrorValue()));
+			const FModifyLobbyJoinPolicy::Params& Params = InAsyncOp.GetParams();
+			const TSharedRef<FLobbyDataEOS>& LobbyData = GetOpDataChecked<TSharedRef<FLobbyDataEOS>>(InAsyncOp, LobbyDataKeyName);
+
+			UE_LOG(LogTemp, Warning, TEXT("[FLobbiesEOSGS::ModifyLobbyJoinPolicy] ModifyLobbyDataImpl Failed: User[%s], Lobby[%s], Result[%s]"),
+				*ToLogString(Params.LocalAccountId), *LobbyData->GetLobbyIdString(), *Result.GetErrorValue().GetLogString());
+			InAsyncOp.SetError(MoveTemp(Result.GetErrorValue()));
 		}
 	})
 	.Then([this](TOnlineAsyncOp<FModifyLobbyJoinPolicy>& InAsyncOp)
 	{
+		const FModifyLobbyJoinPolicy::Params& Params = InAsyncOp.GetParams();
 		const TSharedRef<FLobbyDataEOS>& LobbyData = GetOpDataChecked<TSharedRef<FLobbyDataEOS>>(InAsyncOp, LobbyDataKeyName);
 		const TSharedRef<FClientLobbyDataChanges>& LobbyChanges = GetOpDataChecked<TSharedRef<FClientLobbyDataChanges>>(InAsyncOp, LobbyChangesKeyName);
 
 		// Update local cache and fire events.
 		LobbyData->GetClientLobbyData()->ApplyLobbyUpdateFromLocalChanges(MoveTemp(*LobbyChanges), LobbyEvents);
+
+		UE_LOG(LogTemp, Verbose, TEXT("[FLobbiesEOSGS::ModifyLobbyJoinPolicy] Succeeded: User[%s], Lobby[%s]"),
+			*ToLogString(Params.LocalAccountId), *LobbyData->GetLobbyIdString());
 		InAsyncOp.SetResult(FModifyLobbyJoinPolicy::Result{});
 	})
 	.Enqueue(GetSerialQueue());
@@ -711,20 +955,31 @@ TOnlineAsyncOpHandle<FModifyLobbyAttributes> FLobbiesEOSGS::ModifyLobbyAttribute
 	TOnlineAsyncOpRef<FModifyLobbyAttributes> Op = GetOp<FModifyLobbyAttributes>(MoveTemp(InParams));
 	const FModifyLobbyAttributes::Params& Params = Op->GetParams();
 
-	if (!Services.Get<FAuthEOSGS>()->IsLoggedIn(Params.LocalAccountId))
+	// Check prerequisites.
+	if (!Services.GetAuthInterface()->IsLoggedIn(Params.LocalAccountId))
 	{
+		UE_LOG(LogTemp, Warning, TEXT("[FLobbiesEOSGS::ModifyLobbyAttributes] Failed: User is not logged in. User[%s]"), *ToLogString(Params.LocalAccountId));
 		Op->SetError(Errors::InvalidUser());
 		return Op->GetHandle();
 	}
 
+	if (!Params.LobbyId.IsValid())
+	{
+		UE_LOG(LogTemp, Warning, TEXT("[FLobbiesEOSGS::ModifyLobbyAttributes] Failed: Lobby id is invalid."));
+		Op->SetError(Errors::InvalidParams());
+		return Op->GetHandle();
+	}
+
 	TSharedPtr<FLobbyDataEOS> LobbyData = LobbyDataRegistry->Find(Params.LobbyId);
 	if (!LobbyData)
 	{
+		UE_LOG(LogTemp, Warning, TEXT("[FLobbiesEOSGS::ModifyLobbyAttributes] Failed: Unable to find lobby data. LobbyId[%s]"), *ToLogString(Params.LobbyId));
 		Op->SetError(Errors::InvalidParams());
 		return Op->GetHandle();
 	}
 	Op->Data.Set<TSharedRef<FLobbyDataEOS>>(LobbyDataKeyName, LobbyData.ToSharedRef());
 
+	// Start operation.
 	Op->Then([this](TOnlineAsyncOp<FModifyLobbyAttributes>& InAsyncOp)
 	{
 		const FModifyLobbyAttributes::Params& Params = InAsyncOp.GetParams();
@@ -741,21 +996,29 @@ TOnlineAsyncOpHandle<FModifyLobbyAttributes> FLobbiesEOSGS::ModifyLobbyAttribute
 		ModifyLobbyDataParams.Changes = LobbyChanges;
 		return ModifyLobbyDataImpl(MoveTemp(ModifyLobbyDataParams));
 	})
-	.Then([this](TOnlineAsyncOp<FModifyLobbyAttributes>& InAsyncOp, TDefaultErrorResult<FLobbiesModifyLobbyDataImpl> Result)
+	.Then([this](TOnlineAsyncOp<FModifyLobbyAttributes>& InAsyncOp, TDefaultErrorResult<FLobbiesModifyLobbyDataImpl>&& Result)
 	{
 		if (Result.IsError())
 		{
-			// TODO: Error codes
-			InAsyncOp.SetError(Errors::Unknown(Result.GetErrorValue()));
+			const FModifyLobbyAttributes::Params& Params = InAsyncOp.GetParams();
+			const TSharedRef<FLobbyDataEOS>& LobbyData = GetOpDataChecked<TSharedRef<FLobbyDataEOS>>(InAsyncOp, LobbyDataKeyName);
+
+			UE_LOG(LogTemp, Warning, TEXT("[FLobbiesEOSGS::ModifyLobbyAttributes] ModifyLobbyDataImpl Failed: User[%s], Lobby[%s], Result[%s]"),
+				*ToLogString(Params.LocalAccountId), *LobbyData->GetLobbyIdString(), *Result.GetErrorValue().GetLogString());
+			InAsyncOp.SetError(MoveTemp(Result.GetErrorValue()));
 		}
 	})
 	.Then([this](TOnlineAsyncOp<FModifyLobbyAttributes>& InAsyncOp)
 	{
+		const FModifyLobbyAttributes::Params& Params = InAsyncOp.GetParams();
 		const TSharedRef<FLobbyDataEOS>& LobbyData = GetOpDataChecked<TSharedRef<FLobbyDataEOS>>(InAsyncOp, LobbyDataKeyName);
 		const TSharedRef<FClientLobbyDataChanges>& LobbyChanges = GetOpDataChecked<TSharedRef<FClientLobbyDataChanges>>(InAsyncOp, LobbyChangesKeyName);
 
 		// Update local cache and fire events.
 		LobbyData->GetClientLobbyData()->ApplyLobbyUpdateFromLocalChanges(MoveTemp(*LobbyChanges), LobbyEvents);
+
+		UE_LOG(LogTemp, Verbose, TEXT("[FLobbiesEOSGS::ModifyLobbyAttributes] Succeeded: User[%s], Lobby[%s]"),
+			*ToLogString(Params.LocalAccountId), *LobbyData->GetLobbyIdString());
 		InAsyncOp.SetResult(FModifyLobbyAttributes::Result{});
 	})
 	.Enqueue(GetSerialQueue());
@@ -768,20 +1031,31 @@ TOnlineAsyncOpHandle<FModifyLobbyMemberAttributes> FLobbiesEOSGS::ModifyLobbyMem
 	TOnlineAsyncOpRef<FModifyLobbyMemberAttributes> Op = GetOp<FModifyLobbyMemberAttributes>(MoveTemp(InParams));
 	const FModifyLobbyMemberAttributes::Params& Params = Op->GetParams();
 
-	if (!Services.Get<FAuthEOSGS>()->IsLoggedIn(Params.LocalAccountId))
+	// Check prerequisites.
+	if (!Services.GetAuthInterface()->IsLoggedIn(Params.LocalAccountId))
 	{
+		UE_LOG(LogTemp, Warning, TEXT("[FLobbiesEOSGS::ModifyLobbyMemberAttributes] Failed: User is not logged in. User[%s]"), *ToLogString(Params.LocalAccountId));
 		Op->SetError(Errors::InvalidUser());
 		return Op->GetHandle();
 	}
 
+	if (!Params.LobbyId.IsValid())
+	{
+		UE_LOG(LogTemp, Warning, TEXT("[FLobbiesEOSGS::ModifyLobbyMemberAttributes] Failed: Lobby id is invalid."));
+		Op->SetError(Errors::InvalidParams());
+		return Op->GetHandle();
+	}
+
 	TSharedPtr<FLobbyDataEOS> LobbyData = LobbyDataRegistry->Find(Params.LobbyId);
 	if (!LobbyData)
 	{
+		UE_LOG(LogTemp, Warning, TEXT("[FLobbiesEOSGS::ModifyLobbyMemberAttributes] Failed: Unable to find lobby data. LobbyId[%s]"), *ToLogString(Params.LobbyId));
 		Op->SetError(Errors::InvalidParams());
 		return Op->GetHandle();
 	}
 	Op->Data.Set<TSharedRef<FLobbyDataEOS>>(LobbyDataKeyName, LobbyData.ToSharedRef());
 
+	// Start operation.
 	Op->Then([this](TOnlineAsyncOp<FModifyLobbyMemberAttributes>& InAsyncOp)
 	{
 		const FModifyLobbyMemberAttributes::Params& Params = InAsyncOp.GetParams();
@@ -798,12 +1072,16 @@ TOnlineAsyncOpHandle<FModifyLobbyMemberAttributes> FLobbiesEOSGS::ModifyLobbyMem
 		ModifyLobbyMemberDataParams.Changes = LobbyMemberChanges;
 		return ModifyLobbyMemberDataImpl(MoveTemp(ModifyLobbyMemberDataParams));
 	})
-	.Then([this](TOnlineAsyncOp<FModifyLobbyMemberAttributes>& InAsyncOp, TDefaultErrorResult<FLobbiesModifyLobbyMemberDataImpl> Result)
+	.Then([this](TOnlineAsyncOp<FModifyLobbyMemberAttributes>& InAsyncOp, TDefaultErrorResult<FLobbiesModifyLobbyMemberDataImpl>&& Result)
 	{
 		if (Result.IsError())
 		{
-			// TODO: Error codes
-			InAsyncOp.SetError(Errors::Unknown(Result.GetErrorValue()));
+			const FModifyLobbyMemberAttributes::Params& Params = InAsyncOp.GetParams();
+			const TSharedRef<FLobbyDataEOS>& LobbyData = GetOpDataChecked<TSharedRef<FLobbyDataEOS>>(InAsyncOp, LobbyDataKeyName);
+
+			UE_LOG(LogTemp, Warning, TEXT("[FLobbiesEOSGS::ModifyLobbyMemberAttributes] ModifyLobbyMemberDataImpl Failed: User[%s], Lobby[%s], Result[%s]"),
+				*ToLogString(Params.LocalAccountId), *LobbyData->GetLobbyIdString(), *Result.GetErrorValue().GetLogString());
+			InAsyncOp.SetError(MoveTemp(Result.GetErrorValue()));
 		}
 	})
 	.Then([this](TOnlineAsyncOp<FModifyLobbyMemberAttributes>& InAsyncOp)
@@ -816,6 +1094,9 @@ TOnlineAsyncOpHandle<FModifyLobbyMemberAttributes> FLobbiesEOSGS::ModifyLobbyMem
 		FClientLobbyDataChanges LobbyChanges;
 		LobbyChanges.MutatedMembers.Add(Params.LocalAccountId, LobbyMemberChanges);
 		LobbyData->GetClientLobbyData()->ApplyLobbyUpdateFromLocalChanges(MoveTemp(LobbyChanges), LobbyEvents);
+
+		UE_LOG(LogTemp, Verbose, TEXT("[FLobbiesEOSGS::ModifyLobbyMemberAttributes] Succeeded: User[%s], Lobby[%s]"),
+			*ToLogString(Params.LocalAccountId), *LobbyData->GetLobbyIdString());
 		InAsyncOp.SetResult(FModifyLobbyMemberAttributes::Result{});
 	})
 	.Enqueue(GetSerialQueue());
@@ -854,7 +1135,13 @@ void FLobbiesEOSGS::HandleLobbyUpdated(const EOS_Lobby_LobbyUpdateReceivedCallba
 			if (Result.IsError())
 			{
 				// Todo: handle failure to update lobby from snapshot.
-				UE_LOG(LogTemp, Warning, TEXT("[FLobbiesEOSGS::HandleLobbyUpdated] Failed to apply update. Lobby: %s, Error: %s"), *LobbyData->GetLobbyId(), *Result.GetErrorValue().GetLogString());
+				UE_LOG(LogTemp, Warning, TEXT("[FLobbiesEOSGS::HandleLobbyUpdated] ProcessLobbyNotificationImplOp Failed: Lobby[%s], Result[%s]"),
+					*LobbyData->GetLobbyIdString(), *Result.GetErrorValue().GetLogString());
+			}
+			else
+			{
+				UE_LOG(LogTemp, Verbose, TEXT("[FLobbiesEOSGS::HandleLobbyUpdated] Succeeded: Lobby[%s]"),
+					*LobbyData->GetLobbyIdString());
 			}
 		});
 	}
@@ -874,7 +1161,13 @@ void FLobbiesEOSGS::HandleLobbyMemberUpdated(const EOS_Lobby_LobbyMemberUpdateRe
 			if (Result.IsError())
 			{
 				// Todo: handle failure to update lobby from snapshot.
-				UE_LOG(LogTemp, Warning, TEXT("[FLobbiesEOSGS::HandleLobbyMemberUpdated] Failed to apply update. Lobby: %s, Error: %s"), *LobbyData->GetLobbyId(), *Result.GetErrorValue().GetLogString());
+				UE_LOG(LogTemp, Warning, TEXT("[FLobbiesEOSGS::HandleLobbyMemberUpdated] ProcessLobbyNotificationImplOp Failed: Lobby[%s], Result[%s]"),
+					*LobbyData->GetLobbyIdString(), *Result.GetErrorValue().GetLogString());
+			}
+			else
+			{
+				UE_LOG(LogTemp, Verbose, TEXT("[FLobbiesEOSGS::HandleLobbyMemberUpdated] Succeeded: Lobby[%s]"),
+					*LobbyData->GetLobbyIdString());
 			}
 		});
 	}
@@ -922,7 +1215,13 @@ void FLobbiesEOSGS::HandleLobbyMemberStatusReceived(const EOS_Lobby_LobbyMemberS
 			if (Result.IsError())
 			{
 				// Todo: handle failure to update lobby from snapshot.
-				UE_LOG(LogTemp, Warning, TEXT("[FLobbiesEOSGS::HandleLobbyMemberUpdated] Failed to apply update. Lobby: %s, Error: %s"), *LobbyData->GetLobbyId(), *Result.GetErrorValue().GetLogString());
+				UE_LOG(LogTemp, Warning, TEXT("[FLobbiesEOSGS::HandleLobbyMemberStatusReceived] ProcessLobbyNotificationImplOp Failed: Lobby[%s], Result[%s]"),
+					*LobbyData->GetLobbyIdString(), *Result.GetErrorValue().GetLogString());
+			}
+			else
+			{
+				UE_LOG(LogTemp, Verbose, TEXT("[FLobbiesEOSGS::HandleLobbyMemberStatusReceived] Succeeded: Lobby[%s]"),
+					*LobbyData->GetLobbyIdString());
 			}
 		});
 	}
@@ -930,40 +1229,27 @@ void FLobbiesEOSGS::HandleLobbyMemberStatusReceived(const EOS_Lobby_LobbyMemberS
 
 void FLobbiesEOSGS::HandleLobbyInviteReceived(const EOS_Lobby_LobbyInviteReceivedCallbackInfo* Data)
 {
-#if 0
-	EOS_STRUCT(EOS_Lobby_LobbyInviteReceivedCallbackInfo, (
-		/** Context that was passed into EOS_Lobby_AddNotifyLobbyInviteReceived */
-		void* ClientData;
-		/** The ID of the invitation */
-		const char* InviteId;
-		/** The Product User ID of the local user who received the invitation */
-		EOS_ProductUserId LocalUserId;
-		/** The Product User ID of the user who sent the invitation */
-		EOS_ProductUserId TargetUserId;
-	));
-#endif
-
 	// Todo: Queue this like an operation.
 	const FAccountId LocalAccountId = FindAccountId(Data->LocalUserId);
 	if (LocalAccountId.IsValid())
 	{
 		FLobbyInviteDataEOS::CreateFromInviteId(LobbyPrerequisites.ToSharedRef(), LobbyDataRegistry.ToSharedRef(), LocalAccountId, Data->InviteId, Data->TargetUserId)
-		.Then([this](TFuture<TDefaultErrorResultInternal<TSharedRef<FLobbyInviteDataEOS>>>&& Future)
+		.Next([this](TDefaultErrorResultInternal<TSharedRef<FLobbyInviteDataEOS>>&& Result)
 		{
-			if (Future.Get().IsError())
+			if (Result.IsError())
 			{
 				// Todo: Log / queue a manual fetch of invitations.
-				UE_LOG(LogTemp, Warning, TEXT("[FLobbiesEOSGS::HandleLobbyInviteReceived] Failed to receive invite. Error: %s"),
-					*Future.Get().GetErrorValue().GetLogString());
+				UE_LOG(LogTemp, Warning, TEXT("[FLobbiesEOSGS::HandleLobbyInviteReceived] FLobbyInviteDataEOS::CreateFromInviteId Failed: Result[%s]"),
+					*Result.GetErrorValue().GetLogString());
 			}
 			else
 			{
-				UE_LOG(LogTemp, Warning, TEXT("[FLobbiesEOSGS::HandleLobbyInviteReceived] Received invite. Id: %s, Lobby: %s, Receiver: %s, Sender: %s"),
-					*Future.Get().GetOkValue()->GetInviteId(),
-					*Future.Get().GetOkValue()->GetLobbyData()->GetLobbyId(),
-					*ToLogString(Future.Get().GetOkValue()->GetReceiver()),
-					*ToLogString(Future.Get().GetOkValue()->GetSender()));
-				AddActiveInvite(Future.Get().GetOkValue());
+				UE_LOG(LogTemp, Verbose, TEXT("[FLobbiesEOSGS::HandleLobbyMemberStatusReceived] Succeeded: InviteId[%s], Lobby[%s], Receiver[%s], Sender[%s]"),
+					*Result.GetOkValue()->GetInviteId(),
+					*Result.GetOkValue()->GetLobbyData()->GetLobbyIdString(),
+					*ToLogString(Result.GetOkValue()->GetReceiver()),
+					*ToLogString(Result.GetOkValue()->GetSender()));
+				AddActiveInvite(Result.GetOkValue());
 			}
 		});
 	}
@@ -976,31 +1262,31 @@ void FLobbiesEOSGS::HandleLobbyInviteAccepted(const EOS_Lobby_LobbyInviteAccepte
 	if (LocalAccountId.IsValid())
 	{
 		FLobbyInviteDataEOS::CreateFromInviteId(LobbyPrerequisites.ToSharedRef(), LobbyDataRegistry.ToSharedRef(), LocalAccountId, Data->InviteId, Data->TargetUserId)
-			.Then([this](TFuture<TDefaultErrorResultInternal<TSharedRef<FLobbyInviteDataEOS>>>&& Future)
-				{
-					if (Future.Get().IsError())
-					{
-						// Todo: Log / queue a manual fetch of invitations.
-						UE_LOG(LogTemp, Warning, TEXT("[FLobbiesEOSGS::HandleLobbyInviteReceived] Failed to receive invite. Error: %s"),
-							*Future.Get().GetErrorValue().GetLogString());
-					}
-					else
-					{
-						UE_LOG(LogTemp, Warning, TEXT("[FLobbiesEOSGS::HandleLobbyInviteReceived] Received invite. Id: %s, Lobby: %s, Receiver: %s, Sender: %s"),
-							*Future.Get().GetOkValue()->GetInviteId(),
-							*Future.Get().GetOkValue()->GetLobbyData()->GetLobbyId(),
-							*ToLogString(Future.Get().GetOkValue()->GetReceiver()),
-							*ToLogString(Future.Get().GetOkValue()->GetSender()));
-						const TSharedRef<FLobbyInviteDataEOS>& Invite = Future.Get().GetOkValue();
-
-						LobbyEvents.OnUILobbyJoinRequested.Broadcast(
-							FUILobbyJoinRequested{
-								Invite->GetReceiver(),
-								TResult<TSharedRef<const FLobby>, FOnlineError>(Invite->GetLobbyData()->GetClientLobbyData()->GetPublicDataPtr()),
-								EUILobbyJoinRequestedSource::FromInvitation
-							});
-					}
-				});
+		.Next([this](TDefaultErrorResultInternal<TSharedRef<FLobbyInviteDataEOS>>&& Result)
+		{
+			if (Result.IsError())
+			{
+				// Todo: Log / queue a manual fetch of invitations.
+				UE_LOG(LogTemp, Warning, TEXT("[FLobbiesEOSGS::HandleLobbyInviteAccepted] FLobbyInviteDataEOS::CreateFromInviteId Failed: Result[%s]"),
+					*Result.GetErrorValue().GetLogString());
+			}
+			else
+			{
+				UE_LOG(LogTemp, Verbose, TEXT("[FLobbiesEOSGS::HandleLobbyInviteAccepted] Succeeded: InviteId[%s], Lobby[%s], Receiver[%s], Sender[%s]"),
+					*Result.GetOkValue()->GetInviteId(),
+					*Result.GetOkValue()->GetLobbyData()->GetLobbyIdString(),
+					*ToLogString(Result.GetOkValue()->GetReceiver()),
+					*ToLogString(Result.GetOkValue()->GetSender()));
+
+				const TSharedRef<FLobbyInviteDataEOS>& Invite = Result.GetOkValue();
+				LobbyEvents.OnUILobbyJoinRequested.Broadcast(
+					FUILobbyJoinRequested{
+						Invite->GetReceiver(),
+						TResult<TSharedRef<const FLobby>, FOnlineError>(Invite->GetLobbyData()->GetClientLobbyData()->GetPublicDataPtr()),
+						EUILobbyJoinRequestedSource::FromInvitation
+					});
+			}
+		});
 	}
 }
 
@@ -1015,29 +1301,43 @@ void FLobbiesEOSGS::HandleJoinLobbyAccepted(const EOS_Lobby_JoinLobbyAcceptedCal
 		if (LobbyDetailsResult.IsError())
 		{
 			// Todo: Log / queue a manual fetch of invitations.
-			UE_LOG(LogTemp, Warning, TEXT("[FLobbiesEOSGS::HandleJoinLobbyAccepted] Failed to receive invite. Error: %s"),
-				*LobbyDetailsResult.GetErrorValue().GetLogString());
+			UE_LOG(LogTemp, Warning, TEXT("[FLobbiesEOSGS::HandleJoinLobbyAccepted] FLobbyDetailsEOS::CreateFromUiEventId Failed: User[%s], Result[%s]"),
+				*ToLogString(LocalAccountId), *LobbyDetailsResult.GetErrorValue().GetLogString());
+			return;
 		}
-		LobbyDataRegistry->FindOrCreateFromLobbyDetails(LocalAccountId, LobbyDetailsResult.GetOkValue()).Then([this, LocalAccountId](TFuture<TDefaultErrorResultInternal<TSharedRef<FLobbyDataEOS>>> && Result)
+
+		LobbyDataRegistry->FindOrCreateFromLobbyDetails(LocalAccountId, LobbyDetailsResult.GetOkValue())
+		.Next([this, LocalAccountId](TDefaultErrorResultInternal<TSharedRef<FLobbyDataEOS>> && Result)
 		{
-			LobbyEvents.OnUILobbyJoinRequested.Broadcast(
-				FUILobbyJoinRequested{
-					LocalAccountId,
-					TResult<TSharedRef<const FLobby>, FOnlineError>(Result.Get().GetOkValue()->GetClientLobbyData()->GetPublicDataPtr()),
-					EUILobbyJoinRequestedSource::Unspecified
-				});
+			if (Result.IsError())
+			{
+				UE_LOG(LogTemp, Warning, TEXT("[FLobbiesEOSGS::HandleJoinLobbyAccepted] FLobbyDataRegistryEOS::FindOrCreateFromLobbyDetails Failed: User[%s], Result[%s]"),
+					*ToLogString(LocalAccountId), *Result.GetErrorValue().GetLogString());
+			}
+			else
+			{
+				UE_LOG(LogTemp, Verbose, TEXT("[FLobbiesEOSGS::HandleJoinLobbyAccepted] Succeeded: User[%s], Lobby[%s]"),
+					*ToLogString(LocalAccountId), *Result.GetOkValue()->GetLobbyIdString());
+
+				LobbyEvents.OnUILobbyJoinRequested.Broadcast(
+					FUILobbyJoinRequested{
+						LocalAccountId,
+						TResult<TSharedRef<const FLobby>, FOnlineError>(Result.GetOkValue()->GetClientLobbyData()->GetPublicDataPtr()),
+						EUILobbyJoinRequestedSource::Unspecified
+					});
+			}
 		});
 	}
+	else
+	{
+		UE_LOG(LogTemp, Warning, TEXT("[FLobbiesEOSGS::HandleJoinLobbyAccepted] Failed: Invalid local user."));
+	}
 }
 
 #if !UE_BUILD_SHIPPING
 void FLobbiesEOSGS::CheckMetadata()
 {
 	// Metadata sanity check.
-	ToLogString(FLobbiesJoinLobbyImpl::Params());
-	ToLogString(FLobbiesJoinLobbyImpl::Result());
-	ToLogString(FLobbiesJoinLobbyMemberImpl::Params());
-	ToLogString(FLobbiesJoinLobbyMemberImpl::Result());
 	ToLogString(FLobbiesLeaveLobbyImpl::Params());
 	ToLogString(FLobbiesLeaveLobbyImpl::Result());
 	ToLogString(FLobbiesDestroyLobbyImpl::Params());
@@ -1056,10 +1356,6 @@ void FLobbiesEOSGS::CheckMetadata()
 	ToLogString(FLobbiesModifyLobbyMemberDataImpl::Result());
 	ToLogString(FLobbiesProcessLobbyNotificationImpl::Params());
 	ToLogString(FLobbiesProcessLobbyNotificationImpl::Result());
-	Meta::VisitFields(FLobbiesJoinLobbyImpl::Params(), [](const TCHAR* Name, auto& Field) { return false; });
-	Meta::VisitFields(FLobbiesJoinLobbyImpl::Result(), [](const TCHAR* Name, auto& Field) { return false; });
-	Meta::VisitFields(FLobbiesJoinLobbyMemberImpl::Params(), [](const TCHAR* Name, auto& Field) { return false; });
-	Meta::VisitFields(FLobbiesJoinLobbyMemberImpl::Result(), [](const TCHAR* Name, auto& Field) { return false; });
 	Meta::VisitFields(FLobbiesLeaveLobbyImpl::Params(), [](const TCHAR* Name, auto& Field) { return false; });
 	Meta::VisitFields(FLobbiesLeaveLobbyImpl::Result(), [](const TCHAR* Name, auto& Field) { return false; });
 	Meta::VisitFields(FLobbiesDestroyLobbyImpl::Params(), [](const TCHAR* Name, auto& Field) { return false; });
@@ -1205,224 +1501,22 @@ TSharedPtr<FLobbyInviteDataEOS> FLobbiesEOSGS::GetActiveInvite(FAccountId Target
 	return Result ? *Result : TSharedPtr<FLobbyInviteDataEOS>();
 }
 
-TFuture<TDefaultErrorResult<FLobbiesJoinLobbyImpl>> FLobbiesEOSGS::JoinLobbyImpl(FLobbiesJoinLobbyImpl::Params&& Params)
-{
-	// Check prerequisites.
-	if (!Services.Get<FAuthEOSGS>()->IsLoggedIn(Params.LocalAccountId))
-	{
-		return MakeFulfilledPromise<TDefaultErrorResult<FLobbiesJoinLobbyImpl>>(Errors::InvalidUser()).GetFuture();
-	}
-
-	if (!Params.LobbyData.IsValid())
-	{
-		return MakeFulfilledPromise<TDefaultErrorResult<FLobbiesJoinLobbyImpl>>(Errors::InvalidParams()).GetFuture();
-	}
-
-	// Check whether any of the local users is already in the target lobby.
-	for (const FJoinLobbyLocalUserData& JoinData : Params.LocalUsers)
-	{
-		if (const TSet<TSharedRef<FLobbyDataEOS>>* UserActiveLobbies = ActiveLobbies.Find(JoinData.LocalAccountId))
-		{
-			if (UserActiveLobbies->Find(Params.LobbyData.ToSharedRef()))
-			{
-				return MakeFulfilledPromise<TDefaultErrorResult<FLobbiesJoinLobbyImpl>>(Errors::InvalidParams()).GetFuture();
-			}
-		}
-	}
-
-	TArray<TFuture<TDefaultErrorResult<FLobbiesJoinLobbyMemberImpl>>> PendingMemberJoins;
-	for (FJoinLobbyLocalUserData& UserData : Params.LocalUsers)
-	{
-		FLobbiesJoinLobbyMemberImpl::Params JoinLobbyMemberParams;
-		JoinLobbyMemberParams.LobbyData = Params.LobbyData.ToSharedRef();
-		JoinLobbyMemberParams.LocalAccountId = UserData.LocalAccountId;
-		JoinLobbyMemberParams.Attributes = MoveTemp(UserData.Attributes);
-
-		TSharedRef<TPromise<TDefaultErrorResult<FLobbiesJoinLobbyMemberImpl>>> JoinMemberPromise = MakeShared<TPromise<TDefaultErrorResult<FLobbiesJoinLobbyMemberImpl>>>();
-		JoinLobbyMemberImplOp(MoveTemp(JoinLobbyMemberParams))
-		.OnComplete([JoinMemberPromise](const TOnlineResult<FLobbiesJoinLobbyMemberImpl>& Result)
-		{
-			if (Result.IsOk())
-			{
-				JoinMemberPromise->EmplaceValue(Result.GetOkValue());
-			}
-			else
-			{
-				JoinMemberPromise->EmplaceValue(Result.GetErrorValue());
-			}
-		});
-
-		PendingMemberJoins.Emplace(JoinMemberPromise->GetFuture());
-	}
-
-	TPromise<TDefaultErrorResult<FLobbiesJoinLobbyImpl>> Promise;
-	TFuture<TDefaultErrorResult<FLobbiesJoinLobbyImpl>> Future = Promise.GetFuture();
-
-	WhenAll(MoveTemp(PendingMemberJoins))
-	.Then([this, Promise = MoveTemp(Promise), Params = MoveTemp(Params)](TFuture<TArray<TDefaultErrorResult<FLobbiesJoinLobbyMemberImpl>>>&& Results) mutable
-	{
-		TOptional<FOnlineError> StoredError;
-		for (TDefaultErrorResult<FLobbiesJoinLobbyMemberImpl>& Result : Results.Get())
-		{
-			if (Result.IsError())
-			{
-				// Store first encountered error to return as result.
-				StoredError = Result.GetErrorValue();
-			}
-		}
-
-		if (StoredError)
-		{
-			TArray<TFuture<TDefaultErrorResult<FLobbiesLeaveLobbyImpl>>> PendingMemberExits;
-			for (int32 MemberIndex = 0; MemberIndex < Params.LocalUsers.Num(); ++MemberIndex)
-			{
-				const FAccountId MemberId = Params.LocalUsers[MemberIndex].LocalAccountId;
-
-				if (Results.Get()[MemberIndex].IsError())
-				{
-					FLobbiesLeaveLobbyImpl::Params LeaveLobbyParams;
-					LeaveLobbyParams.LobbyData = Params.LobbyData.ToSharedRef();
-					LeaveLobbyParams.LocalAccountId = MemberId;
-
-					TPromise<TDefaultErrorResult<FLobbiesLeaveLobbyImpl>> LeaveMemberPromise;
-					PendingMemberExits.Emplace(LeaveMemberPromise.GetFuture());
-
-					LeaveLobbyImpl(MoveTemp(LeaveLobbyParams))
-					.Then([LeaveMemberPromise = MoveTemp(LeaveMemberPromise)](TFuture<TDefaultErrorResult<FLobbiesLeaveLobbyImpl>>&& Future) mutable
-					{
-						if (Future.Get().IsError())
-						{
-							// Todo: complain about having an error while handling an error.
-						}
-
-						LeaveMemberPromise.EmplaceValue(MoveTempIfPossible(Future.Get()));
-					});
-				}
-			}
-
-			WhenAll(MoveTemp(PendingMemberExits))
-			.Then([Promise = MoveTemp(Promise), StoredError = MoveTemp(*StoredError)](TFuture<TArray<TDefaultErrorResult<FLobbiesLeaveLobbyImpl>>>&& Future) mutable
-			{
-				Promise.EmplaceValue(MoveTemp(StoredError));
-			});
-		}
-		else
-		{
-			Promise.EmplaceValue(FLobbiesJoinLobbyImpl::Result{});
-		}
-	});
-
-	return Future;
-}
-
-TOnlineAsyncOpHandle<FLobbiesJoinLobbyMemberImpl> FLobbiesEOSGS::JoinLobbyMemberImplOp(FLobbiesJoinLobbyMemberImpl::Params&& InParams)
-{
-	TOnlineAsyncOpRef<FLobbiesJoinLobbyMemberImpl> Op = GetOp<FLobbiesJoinLobbyMemberImpl>(MoveTemp(InParams));
-	const FLobbiesJoinLobbyMemberImpl::Params& Params = Op->GetParams();
-
-	// Setup lobby details - Prefer UI event before invitation before search result.
-	TSharedPtr<FLobbyDetailsEOS> LobbyDetails = Params.LobbyData->GetUserLobbyDetails(Params.LocalAccountId);
-	if (!LobbyDetails)
-	{
-		// Todo: Check whether another local member can invite the user.
-		Op->SetError(Errors::InvalidParams());
-		return Op->GetHandle();
-	}
-
-	if (LobbyDetails->GetInfo()->GetProductVersion() != LobbyPrerequisites->BucketId.GetProductVersion())
-	{
-		Op->SetError(Errors::IncompatibleVersion());
-		return Op->GetHandle();
-	}
-
-	Op->Then([this, LobbyDetails, LobbyData = Params.LobbyData](TOnlineAsyncOp<FLobbiesJoinLobbyMemberImpl>& InAsyncOp, TPromise<const EOS_Lobby_JoinLobbyCallbackInfo*>&& Promise)
-	{
-		const FLobbiesJoinLobbyMemberImpl::Params& Params = InAsyncOp.GetParams();
-
-		EOS_Lobby_JoinLobbyOptions JoinLobbyOptions = {};
-		JoinLobbyOptions.ApiVersion = EOS_LOBBY_JOINLOBBY_API_LATEST;
-		JoinLobbyOptions.LobbyDetailsHandle = LobbyDetails->GetEOSHandle();
-		JoinLobbyOptions.LocalUserId = GetProductUserIdChecked(Params.LocalAccountId);
-		JoinLobbyOptions.bPresenceEnabled = false; // todo
-		JoinLobbyOptions.LocalRTCOptions = nullptr;
-		static_assert(EOS_LOBBY_JOINLOBBY_API_LATEST == 3, "EOS_Lobby_JoinLobbyOptions updated, check new fields");
-		
-		EOS_Async(EOS_Lobby_JoinLobby, LobbyPrerequisites->LobbyInterfaceHandle, JoinLobbyOptions, MoveTemp(Promise));
-	})
-	.Then([this](TOnlineAsyncOp<FLobbiesJoinLobbyMemberImpl>& InAsyncOp, const EOS_Lobby_JoinLobbyCallbackInfo* Data)
-	{
-		if (Data->ResultCode != EOS_EResult::EOS_Success)
-		{
-			// TODO: Error codes
-			InAsyncOp.SetError(Errors::FromEOSResult(Data->ResultCode));
-		}
-	})
-	.Then([this](TOnlineAsyncOp<FLobbiesJoinLobbyMemberImpl>& InAsyncOp)
-	{
-		const FLobbiesJoinLobbyMemberImpl::Params& Params = InAsyncOp.GetParams();
-		FLobbiesModifyLobbyMemberDataImpl::Params ModifyParams;
-		ModifyParams.LobbyData = Params.LobbyData;
-		ModifyParams.LocalAccountId = Params.LocalAccountId;
-		ModifyParams.Changes = MakeShared<FClientLobbyMemberDataChanges>();
-		ModifyParams.Changes->MutatedAttributes = Params.Attributes;
-		return ModifyLobbyMemberDataImpl(MoveTemp(ModifyParams));
-	})
-	.Then([this](TOnlineAsyncOp<FLobbiesJoinLobbyMemberImpl>& InAsyncOp, TDefaultErrorResult<FLobbiesModifyLobbyMemberDataImpl>&& Result)
-	{
-		const FLobbiesJoinLobbyMemberImpl::Params& Params = InAsyncOp.GetParams();
-
-		if (Result.IsError())
-		{
-			// Failed to set attributes - leave the lobby.
-			FLobbiesLeaveLobbyImpl::Params LeaveLobbyParams;
-			LeaveLobbyParams.LobbyData = Params.LobbyData;
-			LeaveLobbyParams.LocalAccountId = Params.LocalAccountId;
-
-			TPromise<void> Promise;
-			TFuture<void> Future = Promise.GetFuture();
-
-			LeaveLobbyImpl(MoveTemp(LeaveLobbyParams))
-			.Then([this, InAsyncOp = InAsyncOp.AsShared(), ErrorResult = MoveTemp(Result.GetErrorValue()), Promise = MoveTemp(Promise)](TFuture<TDefaultErrorResult<FLobbiesLeaveLobbyImpl>>&& Future) mutable
-			{
-				if (Future.Get().IsError())
-				{
-					// Todo: complain about having an error while handling an error.
-				}
-
-				// Todo: Errors.
-				InAsyncOp->SetError(Errors::Unknown(MoveTemp(ErrorResult)));
-				Promise.EmplaceValue();
-			});
-
-			return Future;
-		}
-		else
-		{
-			return MakeFulfilledPromise<void>().GetFuture();
-		}
-	})
-	.Then([this](TOnlineAsyncOp<FLobbiesJoinLobbyMemberImpl>& InAsyncOp)
-	{
-		InAsyncOp.SetResult(FLobbiesJoinLobbyMemberImpl::Result{});
-	})
-	.Enqueue(GetSerialQueue(InParams.LocalAccountId));
-
-	return Op->GetHandle();
-}
-
 TFuture<TDefaultErrorResult<FLobbiesLeaveLobbyImpl>> FLobbiesEOSGS::LeaveLobbyImpl(FLobbiesLeaveLobbyImpl::Params&& Params)
 {
 	// Check prerequisites.
-	if (!Services.Get<FAuthEOSGS>()->IsLoggedIn(Params.LocalAccountId))
+	if (!Services.GetAuthInterface()->IsLoggedIn(Params.LocalAccountId))
 	{
+		UE_LOG(LogTemp, Warning, TEXT("[FLobbiesEOSGS::LeaveLobbyImpl] Failed: User is not logged in. User[%s]"), *ToLogString(Params.LocalAccountId));
 		return MakeFulfilledPromise<TDefaultErrorResult<FLobbiesLeaveLobbyImpl>>(Errors::InvalidUser()).GetFuture();
 	}
 
 	if (!Params.LobbyData.IsValid())
 	{
+		UE_LOG(LogTemp, Warning, TEXT("[FLobbiesEOSGS::LeaveLobbyImpl] Failed: No lobby data provided."));
 		return MakeFulfilledPromise<TDefaultErrorResult<FLobbiesLeaveLobbyImpl>>(Errors::InvalidParams()).GetFuture();
 	}
 
+	// Start operation.
 	EOS_Lobby_LeaveLobbyOptions LeaveLobbyOptions = {};
 	LeaveLobbyOptions.ApiVersion = EOS_LOBBY_LEAVELOBBY_API_LATEST;
 	LeaveLobbyOptions.LocalUserId = GetProductUserIdChecked(Params.LocalAccountId);
@@ -1433,15 +1527,22 @@ TFuture<TDefaultErrorResult<FLobbiesLeaveLobbyImpl>> FLobbiesEOSGS::LeaveLobbyIm
 	TFuture<TDefaultErrorResult<FLobbiesLeaveLobbyImpl>> Future = Promise.GetFuture();
 
 	EOS_Async(EOS_Lobby_LeaveLobby, LobbyPrerequisites->LobbyInterfaceHandle, LeaveLobbyOptions,
-	[Promise = MoveTemp(Promise)](const EOS_Lobby_LeaveLobbyCallbackInfo* Result) mutable
+	[
+		LocalAccountId = MoveTemp(Params.LocalAccountId),
+		LobbyData = MoveTemp(Params.LobbyData),
+		Promise = MoveTemp(Promise)
+	](const EOS_Lobby_LeaveLobbyCallbackInfo* Result) mutable
 	{
 		if (Result->ResultCode != EOS_EResult::EOS_Success)
 		{
-			// Todo: Errors
+			UE_LOG(LogTemp, Warning, TEXT("[FLobbiesEOSGS::DestroyLobbyImpl] EOS_Lobby_DestroyLobby Failed: User[%s], Lobby[%s], Result[%s]"),
+				*ToLogString(LocalAccountId), *LobbyData->GetLobbyIdString(), *LexToString(Result->ResultCode));
 			Promise.EmplaceValue(Errors::FromEOSResult(Result->ResultCode));
 			return;
 		}
 
+		UE_LOG(LogTemp, Verbose, TEXT("[FLobbiesEOSGS::DestroyLobbyImpl] Succeeded: User[%s], Lobby[%s]"),
+			*ToLogString(LocalAccountId), *LobbyData->GetLobbyIdString());
 		Promise.EmplaceValue(FLobbiesLeaveLobbyImpl::Result{});
 	});
 
@@ -1451,35 +1552,46 @@ TFuture<TDefaultErrorResult<FLobbiesLeaveLobbyImpl>> FLobbiesEOSGS::LeaveLobbyIm
 TFuture<TDefaultErrorResult<FLobbiesDestroyLobbyImpl>> FLobbiesEOSGS::DestroyLobbyImpl(FLobbiesDestroyLobbyImpl::Params&& Params)
 {
 	// Check prerequisites.
-	if (!Services.Get<FAuthEOSGS>()->IsLoggedIn(Params.LocalAccountId))
+	if (!Services.GetAuthInterface()->IsLoggedIn(Params.LocalAccountId))
 	{
+		UE_LOG(LogTemp, Warning, TEXT("[FLobbiesEOSGS::DestroyLobbyImpl] Failed: User is not logged in. User[%s]"), *ToLogString(Params.LocalAccountId));
 		return MakeFulfilledPromise<TDefaultErrorResult<FLobbiesDestroyLobbyImpl>>(Errors::InvalidUser()).GetFuture();
 	}
 
-	if (!Params.LobbyData.IsValid())
+	if (Params.LobbyIdString.IsEmpty())
 	{
+		UE_LOG(LogTemp, Warning, TEXT("[FLobbiesEOSGS::DestroyLobbyImpl] Failed: No lobby id provided."));
 		return MakeFulfilledPromise<TDefaultErrorResult<FLobbiesDestroyLobbyImpl>>(Errors::InvalidParams()).GetFuture();
 	}
 
+	// Start operation.
+	FTCHARToUTF8 LobbyIdTranslator(*Params.LobbyIdString);
 	EOS_Lobby_DestroyLobbyOptions DestroyLobbyOptions = {};
 	DestroyLobbyOptions.ApiVersion = EOS_LOBBY_DESTROYLOBBY_API_LATEST;
 	DestroyLobbyOptions.LocalUserId = GetProductUserIdChecked(Params.LocalAccountId);
-	DestroyLobbyOptions.LobbyId = Params.LobbyData->GetLobbyIdEOS();
+	DestroyLobbyOptions.LobbyId = LobbyIdTranslator.Get();
 	static_assert(EOS_LOBBY_DESTROYLOBBY_API_LATEST == 1, "EOS_Lobby_DestroyLobbyOptions updated, check new fields");
 
 	TPromise<TDefaultErrorResult<FLobbiesDestroyLobbyImpl>> Promise;
 	TFuture<TDefaultErrorResult<FLobbiesDestroyLobbyImpl>> Future = Promise.GetFuture();
 
 	EOS_Async(EOS_Lobby_DestroyLobby, LobbyPrerequisites->LobbyInterfaceHandle, DestroyLobbyOptions,
-	[Promise = MoveTemp(Promise)](const EOS_Lobby_DestroyLobbyCallbackInfo* Result) mutable
+	[
+		LocalAccountId = MoveTemp(Params.LocalAccountId),
+		LobbyIdString = MoveTemp(Params.LobbyIdString),
+		Promise = MoveTemp(Promise)
+	](const EOS_Lobby_DestroyLobbyCallbackInfo* Result) mutable
 	{
 		if (Result->ResultCode != EOS_EResult::EOS_Success)
 		{
-			// Todo: Errors
+			UE_LOG(LogTemp, Warning, TEXT("[FLobbiesEOSGS::DestroyLobbyImpl] EOS_Lobby_DestroyLobby Failed: User[%s], Lobby[%s], Result[%s]"),
+				*ToLogString(LocalAccountId), *LobbyIdString, *LexToString(Result->ResultCode));
 			Promise.EmplaceValue(Errors::FromEOSResult(Result->ResultCode));
 			return;
 		}
 
+		UE_LOG(LogTemp, Verbose, TEXT("[FLobbiesEOSGS::DestroyLobbyImpl] Succeeded: User[%s], Lobby[%s]"),
+			*ToLogString(LocalAccountId), *LobbyIdString);
 		Promise.EmplaceValue(FLobbiesDestroyLobbyImpl::Result{});
 	});
 
@@ -1489,16 +1601,19 @@ TFuture<TDefaultErrorResult<FLobbiesDestroyLobbyImpl>> FLobbiesEOSGS::DestroyLob
 TFuture<TDefaultErrorResult<FLobbiesInviteLobbyMemberImpl>> FLobbiesEOSGS::InviteLobbyMemberImpl(FLobbiesInviteLobbyMemberImpl::Params&& Params)
 {
 	// Check prerequisites.
-	if (!Services.Get<FAuthEOSGS>()->IsLoggedIn(Params.LocalAccountId))
+	if (!Services.GetAuthInterface()->IsLoggedIn(Params.LocalAccountId))
 	{
+		UE_LOG(LogTemp, Warning, TEXT("[FLobbiesEOSGS::InviteLobbyMemberImpl] Failed: User is not logged in. User[%s]"), *ToLogString(Params.LocalAccountId));
 		return MakeFulfilledPromise<TDefaultErrorResult<FLobbiesInviteLobbyMemberImpl>>(Errors::InvalidUser()).GetFuture();
 	}
 
 	if (!Params.LobbyData.IsValid())
 	{
+		UE_LOG(LogTemp, Warning, TEXT("[FLobbiesEOSGS::InviteLobbyMemberImpl] Failed: No lobby data provided."));
 		return MakeFulfilledPromise<TDefaultErrorResult<FLobbiesInviteLobbyMemberImpl>>(Errors::InvalidParams()).GetFuture();
 	}
 
+	// Start operation.
 	EOS_Lobby_SendInviteOptions SendInviteOptions = {};
 	SendInviteOptions.ApiVersion = EOS_LOBBY_SENDINVITE_API_LATEST;
 	SendInviteOptions.LocalUserId = GetProductUserIdChecked(Params.LocalAccountId);
@@ -1510,15 +1625,23 @@ TFuture<TDefaultErrorResult<FLobbiesInviteLobbyMemberImpl>> FLobbiesEOSGS::Invit
 	TFuture<TDefaultErrorResult<FLobbiesInviteLobbyMemberImpl>> Future = Promise.GetFuture();
 
 	EOS_Async(EOS_Lobby_SendInvite, LobbyPrerequisites->LobbyInterfaceHandle, SendInviteOptions,
-	[Promise = MoveTemp(Promise)](const EOS_Lobby_SendInviteCallbackInfo* Result) mutable
+	[
+		LocalAccountId = MoveTemp(Params.LocalAccountId),
+		TargetAccountId = MoveTemp(Params.TargetAccountId),
+		LobbyData = MoveTemp(Params.LobbyData),
+		Promise = MoveTemp(Promise)
+	](const EOS_Lobby_SendInviteCallbackInfo* Result) mutable
 	{
 		if (Result->ResultCode != EOS_EResult::EOS_Success)
 		{
-			// Todo: Errors
+			UE_LOG(LogTemp, Warning, TEXT("[FLobbiesEOSGS::InviteLobbyMemberImpl] EOS_Lobby_SendInvite Failed: User[%s], TargetUser[%s], Lobby[%s], Result[%s]"),
+				*ToLogString(LocalAccountId), *ToLogString(TargetAccountId), *LobbyData->GetLobbyIdString(), *LexToString(Result->ResultCode));
 			Promise.EmplaceValue(Errors::FromEOSResult(Result->ResultCode));
 			return;
 		}
 
+		UE_LOG(LogTemp, Verbose, TEXT("[FLobbiesEOSGS::InviteLobbyMemberImpl] Succeeded: User[%s], TargetUser[%s], Lobby[%s]"),
+			*ToLogString(LocalAccountId), *ToLogString(TargetAccountId), *LobbyData->GetLobbyIdString());
 		Promise.EmplaceValue(FLobbiesInviteLobbyMemberImpl::Result{});
 	});
 
@@ -1528,19 +1651,28 @@ TFuture<TDefaultErrorResult<FLobbiesInviteLobbyMemberImpl>> FLobbiesEOSGS::Invit
 TFuture<TDefaultErrorResult<FLobbiesDeclineLobbyInvitationImpl>> FLobbiesEOSGS::DeclineLobbyInvitationImpl(FLobbiesDeclineLobbyInvitationImpl::Params&& Params)
 {
 	// Check prerequisites.
-	if (!Services.Get<FAuthEOSGS>()->IsLoggedIn(Params.LocalAccountId))
+	if (!Services.GetAuthInterface()->IsLoggedIn(Params.LocalAccountId))
 	{
+		UE_LOG(LogTemp, Warning, TEXT("[FLobbiesEOSGS::DeclineLobbyInvitationImpl] Failed: User is not logged in. User[%s]"), *ToLogString(Params.LocalAccountId));
 		return MakeFulfilledPromise<TDefaultErrorResult<FLobbiesDeclineLobbyInvitationImpl>>(Errors::InvalidUser()).GetFuture();
 	}
 
+	if (!Params.LobbyId.IsValid())
+	{
+		UE_LOG(LogTemp, Warning, TEXT("[FLobbiesEOSGS::DeclineLobbyInvitationImpl] Failed: Lobby id is invalid."));
+		return MakeFulfilledPromise<TDefaultErrorResult<FLobbiesDeclineLobbyInvitationImpl>>(Errors::InvalidParams()).GetFuture();
+	}
+
 	// Find the active invitation.
 	TSharedPtr<FLobbyInviteDataEOS> InviteData = GetActiveInvite(Params.LocalAccountId, Params.LobbyId);
 	if (!InviteData)
 	{
-		// Todo: Errors.
-		return MakeFulfilledPromise<TDefaultErrorResult<FLobbiesDeclineLobbyInvitationImpl>>(Errors::Unknown()).GetFuture();
+		UE_LOG(LogTemp, Warning, TEXT("[FLobbiesEOSGS::DeclineLobbyInvitationImpl] Failed: Unable to find invitation for lobby. User[%s], Lobby[%s]"),
+			*ToLogString(Params.LocalAccountId), *ToLogString(Params.LobbyId));
+		return MakeFulfilledPromise<TDefaultErrorResult<FLobbiesDeclineLobbyInvitationImpl>>(Errors::InvalidParams()).GetFuture();
 	}
 
+	// Start operation.
 	EOS_Lobby_RejectInviteOptions RejectInviteOptions = {};
 	RejectInviteOptions.ApiVersion = EOS_LOBBY_REJECTINVITE_API_LATEST;
 	RejectInviteOptions.LocalUserId = GetProductUserIdChecked(Params.LocalAccountId);
@@ -1551,18 +1683,26 @@ TFuture<TDefaultErrorResult<FLobbiesDeclineLobbyInvitationImpl>> FLobbiesEOSGS::
 	TFuture<TDefaultErrorResult<FLobbiesDeclineLobbyInvitationImpl>> Future = Promise.GetFuture();
 
 	EOS_Async(EOS_Lobby_RejectInvite, LobbyPrerequisites->LobbyInterfaceHandle, RejectInviteOptions,
-	[this, Promise = MoveTemp(Promise), InviteData](const EOS_Lobby_RejectInviteCallbackInfo* Result) mutable
+	[
+		this,
+		LocalAccountId = MoveTemp(Params.LocalAccountId),
+		InviteData = MoveTemp(InviteData),
+		Promise = MoveTemp(Promise)
+	](const EOS_Lobby_RejectInviteCallbackInfo* Result) mutable
 	{
 		// Remove active invitation.
 		RemoveActiveInvite(InviteData.ToSharedRef());
 
 		if (Result->ResultCode != EOS_EResult::EOS_Success)
 		{
-			// Todo: Errors
+			UE_LOG(LogTemp, Warning, TEXT("[FLobbiesEOSGS::DeclineLobbyInvitationImpl] EOS_Lobby_RejectInvite Failed: User[%s], Invitation[%s], Lobby[%s], Result[%s]"),
+				*ToLogString(LocalAccountId), *InviteData->GetInviteId(), *InviteData->GetLobbyData()->GetLobbyIdString(), *LexToString(Result->ResultCode));
 			Promise.EmplaceValue(Errors::FromEOSResult(Result->ResultCode));
 			return;
 		}
 
+		UE_LOG(LogTemp, Verbose, TEXT("[FLobbiesEOSGS::DeclineLobbyInvitationImpl] Succeeded: User[%s], Invitation[%s], Lobby[%s]"),
+			*ToLogString(LocalAccountId), *InviteData->GetInviteId(), *InviteData->GetLobbyData()->GetLobbyIdString());
 		Promise.EmplaceValue(FLobbiesDeclineLobbyInvitationImpl::Result{});
 	});
 
@@ -1572,18 +1712,26 @@ TFuture<TDefaultErrorResult<FLobbiesDeclineLobbyInvitationImpl>> FLobbiesEOSGS::
 TFuture<TDefaultErrorResult<FLobbiesKickLobbyMemberImpl>> FLobbiesEOSGS::KickLobbyMemberImpl(FLobbiesKickLobbyMemberImpl::Params&& Params)
 {
 	// Check prerequisites.
-	if (!Services.Get<FAuthEOSGS>()->IsLoggedIn(Params.LocalAccountId))
+	if (!Services.GetAuthInterface()->IsLoggedIn(Params.LocalAccountId))
 	{
+		UE_LOG(LogTemp, Warning, TEXT("[FLobbiesEOSGS::KickLobbyMemberImpl] Failed: User is not logged in. User[%s]"), *ToLogString(Params.LocalAccountId));
 		return MakeFulfilledPromise<TDefaultErrorResult<FLobbiesKickLobbyMemberImpl>>(Errors::InvalidUser()).GetFuture();
 	}
 
 	if (!Params.LobbyData.IsValid())
 	{
+		UE_LOG(LogTemp, Warning, TEXT("[FLobbiesEOSGS::KickLobbyMemberImpl] Failed: No lobby data provided."));
 		return MakeFulfilledPromise<TDefaultErrorResult<FLobbiesKickLobbyMemberImpl>>(Errors::InvalidParams()).GetFuture();
 	}
 
-	// todo: check local user is lobby owner
+	if (Params.LobbyData->GetClientLobbyData()->GetPublicData().OwnerAccountId != Params.LocalAccountId)
+	{
+		UE_LOG(LogTemp, Warning, TEXT("[FLobbiesEOSGS::KickLobbyMemberImpl] Failed: User is not the lobby owner. User[%s], Lobby[%s]"),
+			*ToLogString(Params.LocalAccountId), *Params.LobbyData->GetLobbyIdString());
+		return MakeFulfilledPromise<TDefaultErrorResult<FLobbiesKickLobbyMemberImpl>>(Errors::InvalidParams()).GetFuture();
+	}
 
+	// Start operation.
 	EOS_Lobby_KickMemberOptions KickMemberOptions = {};
 	KickMemberOptions.ApiVersion = EOS_LOBBY_KICKMEMBER_API_LATEST;
 	KickMemberOptions.LocalUserId = GetProductUserIdChecked(Params.LocalAccountId);
@@ -1595,15 +1743,23 @@ TFuture<TDefaultErrorResult<FLobbiesKickLobbyMemberImpl>> FLobbiesEOSGS::KickLob
 	TFuture<TDefaultErrorResult<FLobbiesKickLobbyMemberImpl>> Future = Promise.GetFuture();
 
 	EOS_Async(EOS_Lobby_KickMember, LobbyPrerequisites->LobbyInterfaceHandle, KickMemberOptions,
-	[Promise = MoveTemp(Promise)](const EOS_Lobby_KickMemberCallbackInfo* Result) mutable
+	[
+		LocalAccountId = MoveTemp(Params.LocalAccountId),
+		TargetAccountId = MoveTemp(Params.TargetAccountId),
+		LobbyData = MoveTemp(Params.LobbyData),
+		Promise = MoveTemp(Promise)
+	](const EOS_Lobby_KickMemberCallbackInfo* Result) mutable
 	{
 		if (Result->ResultCode != EOS_EResult::EOS_Success)
 		{
-			// Todo: Errors
+			UE_LOG(LogTemp, Warning, TEXT("[FLobbiesEOSGS::KickLobbyMemberImpl] EOS_Lobby_KickMember Failed: User[%s], TargetUser[%s], Lobby[%s], Result[%s]"),
+				*ToLogString(LocalAccountId), *ToLogString(TargetAccountId), *LobbyData->GetLobbyIdString(), *LexToString(Result->ResultCode));
 			Promise.EmplaceValue(Errors::FromEOSResult(Result->ResultCode));
 			return;
 		}
 
+		UE_LOG(LogTemp, Verbose, TEXT("[FLobbiesEOSGS::KickLobbyMemberImpl] Succeeded: User[%s], TargetUser[%s], Lobby[%s]"),
+			*ToLogString(LocalAccountId), *ToLogString(TargetAccountId), *LobbyData->GetLobbyIdString());
 		Promise.EmplaceValue(FLobbiesKickLobbyMemberImpl::Result{});
 	});
 
@@ -1613,18 +1769,26 @@ TFuture<TDefaultErrorResult<FLobbiesKickLobbyMemberImpl>> FLobbiesEOSGS::KickLob
 TFuture<TDefaultErrorResult<FLobbiesPromoteLobbyMemberImpl>> FLobbiesEOSGS::PromoteLobbyMemberImpl(FLobbiesPromoteLobbyMemberImpl::Params&& Params)
 {
 	// Check prerequisites.
-	if (!Services.Get<FAuthEOSGS>()->IsLoggedIn(Params.LocalAccountId))
+	if (!Services.GetAuthInterface()->IsLoggedIn(Params.LocalAccountId))
 	{
+		UE_LOG(LogTemp, Warning, TEXT("[FLobbiesEOSGS::PromoteLobbyMemberImpl] Failed: User is not logged in. User[%s]"), *ToLogString(Params.LocalAccountId));
 		return MakeFulfilledPromise<TDefaultErrorResult<FLobbiesPromoteLobbyMemberImpl>>(Errors::InvalidUser()).GetFuture();
 	}
 
 	if (!Params.LobbyData.IsValid())
 	{
+		UE_LOG(LogTemp, Warning, TEXT("[FLobbiesEOSGS::PromoteLobbyMemberImpl] Failed: No lobby data provided."));
 		return MakeFulfilledPromise<TDefaultErrorResult<FLobbiesPromoteLobbyMemberImpl>>(Errors::InvalidParams()).GetFuture();
 	}
 
-	// todo: check local user is lobby owner
+	if (Params.LobbyData->GetClientLobbyData()->GetPublicData().OwnerAccountId != Params.LocalAccountId)
+	{
+		UE_LOG(LogTemp, Warning, TEXT("[FLobbiesEOSGS::PromoteLobbyMemberImpl] Failed: User is not the lobby owner. User[%s], Lobby[%s]"),
+			*ToLogString(Params.LocalAccountId), *Params.LobbyData->GetLobbyIdString());
+		return MakeFulfilledPromise<TDefaultErrorResult<FLobbiesPromoteLobbyMemberImpl>>(Errors::InvalidParams()).GetFuture();
+	}
 
+	// Start operation.
 	EOS_Lobby_PromoteMemberOptions PromoteMemberOptions = {};
 	PromoteMemberOptions.ApiVersion = EOS_LOBBY_PROMOTEMEMBER_API_LATEST;
 	PromoteMemberOptions.LocalUserId = GetProductUserIdChecked(Params.LocalAccountId);
@@ -1636,15 +1800,23 @@ TFuture<TDefaultErrorResult<FLobbiesPromoteLobbyMemberImpl>> FLobbiesEOSGS::Prom
 	TFuture<TDefaultErrorResult<FLobbiesPromoteLobbyMemberImpl>> Future = Promise.GetFuture();
 
 	EOS_Async(EOS_Lobby_PromoteMember, LobbyPrerequisites->LobbyInterfaceHandle, PromoteMemberOptions,
-	[Promise = MoveTemp(Promise)](const EOS_Lobby_PromoteMemberCallbackInfo* Result) mutable
+	[
+		LocalAccountId = MoveTemp(Params.LocalAccountId),
+		TargetAccountId = MoveTemp(Params.TargetAccountId),
+		LobbyData = MoveTemp(Params.LobbyData),
+		Promise = MoveTemp(Promise)
+	](const EOS_Lobby_PromoteMemberCallbackInfo* Result) mutable
 	{
 		if (Result->ResultCode != EOS_EResult::EOS_Success)
 		{
-			// Todo: Errors
+			UE_LOG(LogTemp, Warning, TEXT("[FLobbiesEOSGS::PromoteLobbyMemberImpl] EOS_Lobby_PromoteMember Failed: User[%s], TargetUser[%s], Lobby[%s], Result[%s]"),
+				*ToLogString(LocalAccountId), *ToLogString(TargetAccountId), *LobbyData->GetLobbyIdString(), *LexToString(Result->ResultCode));
 			Promise.EmplaceValue(Errors::FromEOSResult(Result->ResultCode));
 			return;
 		}
 
+		UE_LOG(LogTemp, Verbose, TEXT("[FLobbiesEOSGS::PromoteLobbyMemberImpl] Succeeded: User[%s], TargetUser[%s], Lobby[%s]"),
+			*ToLogString(LocalAccountId), *ToLogString(TargetAccountId), *LobbyData->GetLobbyIdString());
 		Promise.EmplaceValue(FLobbiesPromoteLobbyMemberImpl::Result{});
 	});
 
@@ -1654,46 +1826,54 @@ TFuture<TDefaultErrorResult<FLobbiesPromoteLobbyMemberImpl>> FLobbiesEOSGS::Prom
 TFuture<TDefaultErrorResult<FLobbiesModifyLobbyDataImpl>> FLobbiesEOSGS::ModifyLobbyDataImpl(FLobbiesModifyLobbyDataImpl::Params&& Params)
 {
 	// Check prerequisites.
-	if (!Services.Get<FAuthEOSGS>()->IsLoggedIn(Params.LocalAccountId))
+	if (!Services.GetAuthInterface()->IsLoggedIn(Params.LocalAccountId))
 	{
+		UE_LOG(LogTemp, Warning, TEXT("[FLobbiesEOSGS::ModifyLobbyDataImpl] Failed: User is not logged in. User[%s]"), *ToLogString(Params.LocalAccountId));
 		return MakeFulfilledPromise<TDefaultErrorResult<FLobbiesModifyLobbyDataImpl>>(Errors::InvalidUser()).GetFuture();
 	}
 
 	if (!Params.LobbyData.IsValid())
 	{
+		UE_LOG(LogTemp, Warning, TEXT("[FLobbiesEOSGS::ModifyLobbyDataImpl] Failed: No lobby data provided."));
 		return MakeFulfilledPromise<TDefaultErrorResult<FLobbiesModifyLobbyDataImpl>>(Errors::InvalidParams()).GetFuture();
 	}
 
 	TSharedPtr<FLobbyDetailsEOS> LobbyDetails = Params.LobbyData->GetUserLobbyDetails(Params.LocalAccountId);
 	if (!LobbyDetails)
 	{
+		UE_LOG(LogTemp, Warning, TEXT("[FLobbiesEOSGS::ModifyLobbyDataImpl] Failed: Unable to find lobby details for user. User[%s], Lobby[%s]"),
+			*ToLogString(Params.LocalAccountId), *Params.LobbyData->GetLobbyIdString());
 		return MakeFulfilledPromise<TDefaultErrorResult<FLobbiesModifyLobbyDataImpl>>(Errors::InvalidParams()).GetFuture();
 	}
 
 	if (Params.LocalAccountId != Params.LobbyData->GetClientLobbyData()->GetPublicData().OwnerAccountId)
 	{
+		UE_LOG(LogTemp, Warning, TEXT("[FLobbiesEOSGS::ModifyLobbyDataImpl] Failed: User is not lobby owner. User[%s], Lobby[%s]"),
+			*ToLogString(Params.LocalAccountId), *Params.LobbyData->GetLobbyIdString());
 		return MakeFulfilledPromise<TDefaultErrorResult<FLobbiesModifyLobbyDataImpl>>(Errors::InvalidParams()).GetFuture();
 	}
 
-	UE_LOG(LogTemp, Warning, TEXT("[FLobbiesEOSGS::ModifyLobbyDataImpl] Start. Lobby: %s, Member: %s"),
-		*Params.LobbyData->GetLobbyId(), *ToLogString(Params.LocalAccountId));
-
+	// Start operation.
 	TPromise<TDefaultErrorResult<FLobbiesModifyLobbyDataImpl>> Promise;
 	TFuture<TDefaultErrorResult<FLobbiesModifyLobbyDataImpl>> Future = Promise.GetFuture();
 
 	LobbyDetails->ApplyLobbyDataUpdateFromLocalChanges(Params.LocalAccountId, *Params.Changes)
-	.Then([Promise = MoveTemp(Promise)](TFuture<EOS_EResult>&& Future) mutable
+	.Next([
+		LocalAccountId = Params.LocalAccountId,
+		LobbyData = Params.LobbyData,
+		Promise = MoveTemp(Promise)
+	](EOS_EResult EosResult) mutable
 	{
-		UE_LOG(LogTemp, Warning, TEXT("[FLobbiesEOSGS::ModifyLobbyDataImpl] Complete. Result: %s"), *LexToString(Future.Get()));
-
-		// Todo: Handle "no change" better.
-		if (Future.Get() != EOS_EResult::EOS_Success && Future.Get() != EOS_EResult::EOS_NoChange)
+		if (EosResult != EOS_EResult::EOS_Success && EosResult != EOS_EResult::EOS_NoChange)
 		{
-			// Todo: Errors
-			Promise.EmplaceValue(Errors::FromEOSResult(Future.Get()));
+			UE_LOG(LogTemp, Warning, TEXT("[FLobbiesEOSGS::ModifyLobbyDataImpl] ApplyLobbyDataUpdateFromLocalChanges Failed: User[%s], Lobby[%s], Result[%s]"),
+				*ToLogString(LocalAccountId), *LobbyData->GetLobbyIdString(), *LexToString(EosResult));
+			Promise.EmplaceValue(Errors::FromEOSResult(EosResult));
 			return;
 		}
 
+		UE_LOG(LogTemp, Verbose, TEXT("[FLobbiesEOSGS::ModifyLobbyDataImpl] Succeeded: User[%s], Lobby[%s]"),
+			*ToLogString(LocalAccountId), *LobbyData->GetLobbyIdString());
 		Promise.EmplaceValue(FLobbiesModifyLobbyDataImpl::Result{});
 	});
 
@@ -1703,36 +1883,47 @@ TFuture<TDefaultErrorResult<FLobbiesModifyLobbyDataImpl>> FLobbiesEOSGS::ModifyL
 TFuture<TDefaultErrorResult<FLobbiesModifyLobbyMemberDataImpl>> FLobbiesEOSGS::ModifyLobbyMemberDataImpl(FLobbiesModifyLobbyMemberDataImpl::Params&& Params)
 {
 	// Check prerequisites.
-	if (!Services.Get<FAuthEOSGS>()->IsLoggedIn(Params.LocalAccountId))
+	if (!Services.GetAuthInterface()->IsLoggedIn(Params.LocalAccountId))
 	{
+		UE_LOG(LogTemp, Warning, TEXT("[FLobbiesEOSGS::ModifyLobbyMemberDataImpl] Failed: User is not logged in. User[%s]"), *ToLogString(Params.LocalAccountId));
 		return MakeFulfilledPromise<TDefaultErrorResult<FLobbiesModifyLobbyMemberDataImpl>>(Errors::InvalidUser()).GetFuture();
 	}
 
 	if (!Params.LobbyData.IsValid())
 	{
+		UE_LOG(LogTemp, Warning, TEXT("[FLobbiesEOSGS::ModifyLobbyMemberDataImpl] Failed: No lobby data provided."));
 		return MakeFulfilledPromise<TDefaultErrorResult<FLobbiesModifyLobbyMemberDataImpl>>(Errors::InvalidParams()).GetFuture();
 	}
 
 	TSharedPtr<FLobbyDetailsEOS> LobbyDetails = Params.LobbyData->GetUserLobbyDetails(Params.LocalAccountId);
 	if (!LobbyDetails)
 	{
+		UE_LOG(LogTemp, Warning, TEXT("[FLobbiesEOSGS::ModifyLobbyMemberDataImpl] Failed: Unable to find lobby details for user. User[%s], Lobby[%s]"),
+			*ToLogString(Params.LocalAccountId), *Params.LobbyData->GetLobbyIdString());
 		return MakeFulfilledPromise<TDefaultErrorResult<FLobbiesModifyLobbyMemberDataImpl>>(Errors::InvalidParams()).GetFuture();
 	}
 
+	// Start operation.
 	TPromise<TDefaultErrorResult<FLobbiesModifyLobbyMemberDataImpl>> Promise;
 	TFuture<TDefaultErrorResult<FLobbiesModifyLobbyMemberDataImpl>> Future = Promise.GetFuture();
 
 	LobbyDetails->ApplyLobbyMemberDataUpdateFromLocalChanges(Params.LocalAccountId, *Params.Changes)
-	.Then([Promise = MoveTemp(Promise)](TFuture<EOS_EResult>&& Future) mutable
+	.Next([
+		LocalAccountId = Params.LocalAccountId,
+		LobbyData = Params.LobbyData,
+		Promise = MoveTemp(Promise)
+	](EOS_EResult EosResult) mutable
 	{
-		// Todo: Handle "no change" better.
-		if (Future.Get() != EOS_EResult::EOS_Success && Future.Get() != EOS_EResult::EOS_NoChange)
+		if (EosResult != EOS_EResult::EOS_Success && EosResult != EOS_EResult::EOS_NoChange)
 		{
-			// Todo: Errors
-			Promise.EmplaceValue(Errors::FromEOSResult(Future.Get()));
+			UE_LOG(LogTemp, Warning, TEXT("[FLobbiesEOSGS::ModifyLobbyMemberDataImpl] ApplyLobbyMemberDataUpdateFromLocalChanges Failed: User[%s], Lobby[%s], Result[%s]"),
+				*ToLogString(LocalAccountId), *LobbyData->GetLobbyIdString(), *LexToString(EosResult));
+			Promise.EmplaceValue(Errors::FromEOSResult(EosResult));
 			return;
 		}
 
+		UE_LOG(LogTemp, Verbose, TEXT("[FLobbiesEOSGS::ModifyLobbyMemberDataImpl] Succeeded: User[%s], Lobby[%s]"),
+			*ToLogString(LocalAccountId), *LobbyData->GetLobbyIdString());
 		Promise.EmplaceValue(FLobbiesModifyLobbyMemberDataImpl::Result{});
 	});
 
@@ -1744,12 +1935,15 @@ TOnlineAsyncOpHandle<FLobbiesProcessLobbyNotificationImpl> FLobbiesEOSGS::Proces
 	TOnlineAsyncOpRef<FLobbiesProcessLobbyNotificationImpl> Op = GetOp<FLobbiesProcessLobbyNotificationImpl>(MoveTemp(InParams));
 	const FLobbiesProcessLobbyNotificationImpl::Params& Params = Op->GetParams();
 
+	// Check prerequisites.
 	if (!Params.LobbyData.IsValid())
 	{
+		UE_LOG(LogTemp, Warning, TEXT("[FLobbiesEOSGS::ProcessLobbyNotificationImplOp] Failed: No lobby data provided."));
 		Op->SetError(Errors::InvalidParams());
 		return Op->GetHandle();
 	}
 
+	// Start operation.
 	Op->Then([this](TOnlineAsyncOp<FLobbiesProcessLobbyNotificationImpl>& InAsyncOp)
 	{
 		const FLobbiesProcessLobbyNotificationImpl::Params& Params = InAsyncOp.GetParams();
@@ -1759,10 +1953,10 @@ TOnlineAsyncOpHandle<FLobbiesProcessLobbyNotificationImpl> FLobbiesEOSGS::Proces
 		TSharedPtr<FLobbyDetailsEOS> LobbyDetails = Params.LobbyData->GetActiveLobbyDetails();
 		if (!LobbyDetails.IsValid())
 		{
-			UE_LOG(LogTemp, Log, TEXT("[FLobbiesEOSGS::ProcessLobbyNotificationImplOp] Failed to find active lobby details to process lobby notificaions: Lobby: %s"),
-				*Params.LobbyData->GetLobbyId());
-			InAsyncOp.SetError(Errors::Unknown());
-			return MakeFulfilledPromise<TDefaultErrorResultInternal<TSharedRef<FClientLobbySnapshot>>>().GetFuture();
+			UE_LOG(LogTemp, Warning, TEXT("[FLobbiesEOSGS::ProcessLobbyNotificationImplOp] Failed: Unable to find active lobby details to process notificaions. Lobby[%s]"),
+				*Params.LobbyData->GetLobbyIdString());
+			InAsyncOp.SetError(Errors::InvalidState());
+			return MakeFulfilledPromise<TDefaultErrorResultInternal<TSharedRef<FClientLobbySnapshot>>>(Errors::InvalidState()).GetFuture();
 		}
 
 		InAsyncOp.Data.Set<TSharedRef<FLobbyDetailsEOS>>(LobbyDetailsKeyName, LobbyDetails.ToSharedRef());
@@ -1777,8 +1971,9 @@ TOnlineAsyncOpHandle<FLobbiesProcessLobbyNotificationImpl> FLobbiesEOSGS::Proces
 
 		if (LobbySnapshotResult.IsError())
 		{
-			// Todo: errors.
-			InAsyncOp.SetError(Errors::Unknown(MoveTempIfPossible(LobbySnapshotResult.GetErrorValue())));
+			UE_LOG(LogTemp, Warning, TEXT("[FLobbiesEOSGS::ProcessLobbyNotificationImplOp] GetLobbySnapshot Failed. Lobby[%s], Result[%s]"),
+				*Params.LobbyData->GetLobbyIdString(), *LobbySnapshotResult.GetErrorValue().GetLogString());
+			InAsyncOp.SetError(MoveTemp(LobbySnapshotResult.GetErrorValue()));
 			return;
 		}
 
@@ -1793,8 +1988,9 @@ TOnlineAsyncOpHandle<FLobbiesProcessLobbyNotificationImpl> FLobbiesEOSGS::Proces
 				TDefaultErrorResultInternal<TSharedRef<FClientLobbyMemberSnapshot>> MemberSnapshotResult = LobbyDetails->GetLobbyMemberSnapshot(MutatedMemberAccountId);
 				if (MemberSnapshotResult.IsError())
 				{
-					// Todo: errors.
-					InAsyncOp.SetError(Errors::Unknown(MoveTempIfPossible(MemberSnapshotResult.GetErrorValue())));
+					UE_LOG(LogTemp, Warning, TEXT("[FLobbiesEOSGS::ProcessLobbyNotificationImplOp] GetLobbyMemberSnapshot Failed. User[%s], Lobby[%s], Result[%s]"),
+						*ToLogString(MutatedMemberAccountId), *Params.LobbyData->GetLobbyIdString(), *LobbySnapshotResult.GetErrorValue().GetLogString());
+					InAsyncOp.SetError(MoveTemp(MemberSnapshotResult.GetErrorValue()));
 					return;
 				}
 				LobbyMemberSnapshots.Add(MutatedMemberAccountId, MoveTemp(MemberSnapshotResult.GetOkValue()));
@@ -1826,6 +2022,8 @@ TOnlineAsyncOpHandle<FLobbiesProcessLobbyNotificationImpl> FLobbiesEOSGS::Proces
 			RemoveActiveLobby(LeavingMember, Params.LobbyData.ToSharedRef());
 		}
 
+		UE_LOG(LogTemp, Verbose, TEXT("[FLobbiesEOSGS::ProcessLobbyNotificationImplOp] Succeeded: Lobby[%s]"),
+			*Params.LobbyData->GetLobbyIdString());
 		InAsyncOp.SetResult(FLobbiesProcessLobbyNotificationImpl::Result{});
 	})
 	.Enqueue(GetSerialQueue());
diff --git a/Engine/Plugins/Online/Experimental/OnlineServicesEOSGS/Source/Private/Online/LobbiesEOSGSTypes.cpp b/Engine/Plugins/Online/Experimental/OnlineServicesEOSGS/Source/Private/Online/LobbiesEOSGSTypes.cpp
--- a/Engine/Plugins/Online/Experimental/OnlineServicesEOSGS/Source/Private/Online/LobbiesEOSGSTypes.cpp
+++ b/Engine/Plugins/Online/Experimental/OnlineServicesEOSGS/Source/Private/Online/LobbiesEOSGSTypes.cpp
@@ -9,21 +9,6 @@ namespace UE::Online {
 
 namespace Private {
 
-FString TranslateLobbyId(EOS_LobbyId EOSLobbyId)
-{
-	return FString(UTF8_TO_TCHAR(EOSLobbyId));
-}
-
-FString TranslateLobbyInviteId(const char* InviteId)
-{
-	return FString(UTF8_TO_TCHAR(InviteId));
-}
-
-FLobbyAttributeId TranslateLobbyAttributeId(const char* AttributeId)
-{
-	return FLobbyAttributeId(UTF8_TO_TCHAR(AttributeId));
-}
-
 EOS_EComparisonOp TranslateSearchComparison(ELobbyComparisonOp Op)
 {
 	switch (Op)
@@ -48,6 +33,11 @@ EOS_EComparisonOp TranslateSearchComparison(ELobbyComparisonOp Op)
 
 } // Private
 
+FString ToLogString(const FLobbyBucketIdEOS& BucketId)
+{
+	return FString::Printf(TEXT("%s:%d"), *BucketId.GetProductName(), BucketId.GetProductVersion());
+}
+
 const FString FLobbyBucketIdEOS::Separator = TEXT("|");
 
 FLobbyBucketIdEOS::FLobbyBucketIdEOS(FString ProductName, int32 ProductVersion)
@@ -94,7 +84,7 @@ FLobbyAttributeTranslator<ELobbyTranslationType::ToService>::FLobbyAttributeTran
 
 FLobbyAttributeTranslator<ELobbyTranslationType::FromService>::FLobbyAttributeTranslator(const EOS_Lobby_AttributeData& FromAttributeData)
 {
-	FLobbyAttributeId AttributeId = Private::TranslateLobbyAttributeId(FromAttributeData.Key);
+	FLobbyAttributeId AttributeId(UTF8_TO_TCHAR(FromAttributeData.Key));
 	FLobbyVariant VariantData;
 
 	switch (FromAttributeData.ValueType)
@@ -148,32 +138,46 @@ const EOS_HLobbyDetails FLobbyDetailsEOS::InvalidLobbyDetailsHandle = {};
 TDefaultErrorResultInternal<TSharedRef<FLobbyDetailsEOS>> FLobbyDetailsEOS::CreateFromLobbyId(
 	const TSharedRef<FLobbyPrerequisitesEOS>& Prerequisites,
 	FAccountId LocalAccountId,
-	EOS_LobbyId LobbyId)
+	EOS_LobbyId LobbyIdEOS)
 {
 	EOS_HLobbyDetails LobbyDetailsHandle = {};
 
 	EOS_Lobby_CopyLobbyDetailsHandleOptions Options;
 	Options.ApiVersion = EOS_LOBBY_COPYLOBBYDETAILSHANDLE_API_LATEST;
-	Options.LobbyId = LobbyId;
+	Options.LobbyId = LobbyIdEOS;
 	Options.LocalUserId = GetProductUserIdChecked(LocalAccountId);
 	static_assert(EOS_LOBBY_COPYLOBBYDETAILSHANDLE_API_LATEST == 1, "EOS_Lobby_CopyLobbyDetailsHandleOptions updated, check new fields");
 
 	EOS_EResult EOSResult = EOS_Lobby_CopyLobbyDetailsHandle(Prerequisites->LobbyInterfaceHandle, &Options, &LobbyDetailsHandle);
 	if (EOSResult != EOS_EResult::EOS_Success)
 	{
+		UE_LOG(LogTemp, Warning, TEXT("[FLobbyDetailsEOS::CreateFromLobbyId] EOS_Lobby_CopyLobbyDetailsHandle Failed: User[%s], Lobby[%s], Result[%s]"),
+			*ToLogString(LocalAccountId), UTF8_TO_TCHAR(LobbyIdEOS), *ToLogString(EOSResult));
 		return TDefaultErrorResultInternal<TSharedRef<FLobbyDetailsEOS>>(Errors::FromEOSResult(EOSResult));
 	}
 
 	TDefaultErrorResultInternal<TSharedRef<FLobbyDetailsInfoEOS>> LobbyDetailsInfoResult = FLobbyDetailsInfoEOS::Create(LobbyDetailsHandle);
 	if (LobbyDetailsInfoResult.IsError())
 	{
+		UE_LOG(LogTemp, Warning, TEXT("[FLobbyDetailsEOS::CreateFromLobbyId] FLobbyDetailsInfoEOS::Create Failed: User[%s], Lobby[%s], Result[%s]"),
+			*ToLogString(LocalAccountId), UTF8_TO_TCHAR(LobbyIdEOS), *ToLogString(EOSResult));
 		EOS_LobbyDetails_Release(LobbyDetailsHandle);
 		return TDefaultErrorResultInternal<TSharedRef<FLobbyDetailsEOS>>(MoveTemp(LobbyDetailsInfoResult.GetErrorValue()));
 	}
 
+	UE_LOG(LogTemp, VeryVerbose, TEXT("[FLobbyDetailsEOS::CreateFromLobbyId] Succeeded: User[%s], Lobby[%s]"),
+		*ToLogString(LocalAccountId), UTF8_TO_TCHAR(LobbyIdEOS));
 	return TDefaultErrorResultInternal<TSharedRef<FLobbyDetailsEOS>>(MakeShared<FLobbyDetailsEOS>(Prerequisites, LobbyDetailsInfoResult.GetOkValue(), LocalAccountId, ELobbyDetailsSource::Active, LobbyDetailsHandle));
 }
 
+TDefaultErrorResultInternal<TSharedRef<FLobbyDetailsEOS>> FLobbyDetailsEOS::CreateFromLobbyId(
+	const TSharedRef<FLobbyPrerequisitesEOS>& Prerequisites,
+	FAccountId LocalAccountId,
+	const FString& LobbyId)
+{
+	return CreateFromLobbyId(Prerequisites, LocalAccountId, TCHAR_TO_UTF8(*LobbyId));
+}
+
 TDefaultErrorResultInternal<TSharedRef<FLobbyDetailsEOS>> FLobbyDetailsEOS::CreateFromInviteId(
 	const TSharedRef<FLobbyPrerequisitesEOS>& Prerequisites,
 	FAccountId LocalAccountId,
@@ -189,16 +193,22 @@ TDefaultErrorResultInternal<TSharedRef<FLobbyDetailsEOS>> FLobbyDetailsEOS::Crea
 	EOS_EResult EOSResult = EOS_Lobby_CopyLobbyDetailsHandleByInviteId(Prerequisites->LobbyInterfaceHandle, &Options, &LobbyDetailsHandle);
 	if (EOSResult != EOS_EResult::EOS_Success)
 	{
+		UE_LOG(LogTemp, Warning, TEXT("[FLobbyDetailsEOS::CreateFromInviteId] EOS_Lobby_CopyLobbyDetailsHandleByInviteId Failed: User[%s], InviteId[%s], Result[%s]"),
+			*ToLogString(LocalAccountId), UTF8_TO_TCHAR(InviteId), *ToLogString(EOSResult));
 		return TDefaultErrorResultInternal<TSharedRef<FLobbyDetailsEOS>>(Errors::FromEOSResult(EOSResult));
 	}
 
 	TDefaultErrorResultInternal<TSharedRef<FLobbyDetailsInfoEOS>> LobbyDetailsInfoResult = FLobbyDetailsInfoEOS::Create(LobbyDetailsHandle);
 	if (LobbyDetailsInfoResult.IsError())
 	{
+		UE_LOG(LogTemp, Warning, TEXT("[FLobbyDetailsEOS::CreateFromInviteId] FLobbyDetailsInfoEOS::Create Failed: User[%s], InviteId[%s], Result[%s]"),
+			*ToLogString(LocalAccountId), UTF8_TO_TCHAR(InviteId), *LobbyDetailsInfoResult.GetErrorValue().GetLogString());
 		EOS_LobbyDetails_Release(LobbyDetailsHandle);
 		return TDefaultErrorResultInternal<TSharedRef<FLobbyDetailsEOS>>(MoveTemp(LobbyDetailsInfoResult.GetErrorValue()));
 	}
 
+	UE_LOG(LogTemp, VeryVerbose, TEXT("[FLobbyDetailsEOS::CreateFromInviteId] Succeeded: User[%s], InviteId[%s], Lobby[%s]"),
+		*ToLogString(LocalAccountId), UTF8_TO_TCHAR(InviteId), UTF8_TO_TCHAR(LobbyDetailsInfoResult.GetOkValue()->GetLobbyIdEOS()));
 	return TDefaultErrorResultInternal<TSharedRef<FLobbyDetailsEOS>>(MakeShared<FLobbyDetailsEOS>(Prerequisites, LobbyDetailsInfoResult.GetOkValue(), LocalAccountId, ELobbyDetailsSource::Invite, LobbyDetailsHandle));
 }
 
@@ -217,16 +227,22 @@ TDefaultErrorResultInternal<TSharedRef<FLobbyDetailsEOS>> FLobbyDetailsEOS::Crea
 	EOS_EResult EOSResult = EOS_Lobby_CopyLobbyDetailsHandleByUiEventId(Prerequisites->LobbyInterfaceHandle, &Options, &LobbyDetailsHandle);
 	if (EOSResult != EOS_EResult::EOS_Success)
 	{
+		UE_LOG(LogTemp, Warning, TEXT("[FLobbyDetailsEOS::CreateFromUiEventId] EOS_Lobby_CopyLobbyDetailsHandleByUiEventId Failed: User[%s], Result[%s]"),
+			*ToLogString(LocalAccountId), *ToLogString(EOSResult));
 		return TDefaultErrorResultInternal<TSharedRef<FLobbyDetailsEOS>>(Errors::FromEOSResult(EOSResult));
 	}
 
 	TDefaultErrorResultInternal<TSharedRef<FLobbyDetailsInfoEOS>> LobbyDetailsInfoResult = FLobbyDetailsInfoEOS::Create(LobbyDetailsHandle);
 	if (LobbyDetailsInfoResult.IsError())
 	{
+		UE_LOG(LogTemp, Warning, TEXT("[FLobbyDetailsEOS::CreateFromUiEventId] FLobbyDetailsInfoEOS::Create Failed: User[%s], Result[%s]"),
+			*ToLogString(LocalAccountId), *LobbyDetailsInfoResult.GetErrorValue().GetLogString());
 		EOS_LobbyDetails_Release(LobbyDetailsHandle);
 		return TDefaultErrorResultInternal<TSharedRef<FLobbyDetailsEOS>>(MoveTemp(LobbyDetailsInfoResult.GetErrorValue()));
 	}
 
+	UE_LOG(LogTemp, VeryVerbose, TEXT("[FLobbyDetailsEOS::CreateFromUiEventId] Succeeded: User[%s], Lobby[%s]"),
+		*ToLogString(LocalAccountId), UTF8_TO_TCHAR(LobbyDetailsInfoResult.GetOkValue()->GetLobbyIdEOS()));
 	return TDefaultErrorResultInternal<TSharedRef<FLobbyDetailsEOS>>(MakeShared<FLobbyDetailsEOS>(Prerequisites, LobbyDetailsInfoResult.GetOkValue(), LocalAccountId, ELobbyDetailsSource::UiEvent, LobbyDetailsHandle));
 }
 
@@ -246,16 +262,22 @@ TDefaultErrorResultInternal<TSharedRef<FLobbyDetailsEOS>> FLobbyDetailsEOS::Crea
 	EOS_EResult EOSResult = EOS_LobbySearch_CopySearchResultByIndex(SearchHandle, &Options, &LobbyDetailsHandle);
 	if (EOSResult != EOS_EResult::EOS_Success)
 	{
+		UE_LOG(LogTemp, Warning, TEXT("[FLobbyDetailsEOS::CreateFromSearchResult] EOS_LobbySearch_CopySearchResultByIndex Failed: User[%s], Result[%s]"),
+			*ToLogString(LocalAccountId), *ToLogString(EOSResult));
 		return TDefaultErrorResultInternal<TSharedRef<FLobbyDetailsEOS>>(Errors::FromEOSResult(EOSResult));
 	}
 
 	TDefaultErrorResultInternal<TSharedRef<FLobbyDetailsInfoEOS>> LobbyDetailsInfoResult = FLobbyDetailsInfoEOS::Create(LobbyDetailsHandle);
 	if (LobbyDetailsInfoResult.IsError())
 	{
+		UE_LOG(LogTemp, Warning, TEXT("[FLobbyDetailsEOS::CreateFromSearchResult] FLobbyDetailsInfoEOS::Create Failed: User[%s], Result[%s]"),
+			*ToLogString(LocalAccountId), *LobbyDetailsInfoResult.GetErrorValue().GetLogString());
 		EOS_LobbyDetails_Release(LobbyDetailsHandle);
 		return TDefaultErrorResultInternal<TSharedRef<FLobbyDetailsEOS>>(MoveTemp(LobbyDetailsInfoResult.GetErrorValue()));
 	}
 
+	UE_LOG(LogTemp, VeryVerbose, TEXT("[FLobbyDetailsEOS::CreateFromSearchResult] Succeeded: User[%s], Lobby[%s]"),
+		*ToLogString(LocalAccountId), UTF8_TO_TCHAR(LobbyDetailsInfoResult.GetOkValue()->GetLobbyIdEOS()));
 	return TDefaultErrorResultInternal<TSharedRef<FLobbyDetailsEOS>>(MakeShared<FLobbyDetailsEOS>(Prerequisites, LobbyDetailsInfoResult.GetOkValue(), LocalAccountId, ELobbyDetailsSource::Search, LobbyDetailsHandle));
 }
 
@@ -264,11 +286,18 @@ FLobbyDetailsEOS::~FLobbyDetailsEOS()
 	EOS_LobbyDetails_Release(LobbyDetailsHandle);
 }
 
+bool FLobbyDetailsEOS::IsBucketCompatible() const
+{
+	return LobbyDetailsInfo->GetBucketId() == Prerequisites->BucketId;
+}
+
 TFuture<TDefaultErrorResultInternal<TSharedRef<FClientLobbySnapshot>>> FLobbyDetailsEOS::GetLobbySnapshot() const
 {
 	TSharedPtr<FAuthEOSGS> AuthInterface = Prerequisites->AuthInterface.Pin();
 	if (!AuthInterface)
 	{
+		UE_LOG(LogTemp, Warning, TEXT("[FLobbyDetailsEOS::GetLobbySnapshot] Failed: Auth interface missing. Lobby[%s]"),
+			UTF8_TO_TCHAR(GetInfo()->GetLobbyIdEOS()));
 		return MakeFulfilledPromise<TDefaultErrorResultInternal<TSharedRef<FClientLobbySnapshot>>>(Errors::MissingInterface()).GetFuture();
 	}
 
@@ -296,12 +325,12 @@ TFuture<TDefaultErrorResultInternal<TSharedRef<FClientLobbySnapshot>>> FLobbyDet
 
 	// Resolve lobby member product user ids to FAccountId before proceeding.
 	AuthInterface->ResolveAccountIds(AssociatedLocalUser, *MemberProductUserIds)
-	.Then([StrongThis = AsShared(), Promise = MoveTemp(Promise), MemberProductUserIds](TFuture<TArray<FAccountId>>&& Future) mutable
+	.Next([StrongThis = AsShared(), Promise = MoveTemp(Promise), MemberProductUserIds](TArray<FAccountId>&& ResolvedAccountIds) mutable
 	{
-		const TArray<FAccountId>& ResolvedAccountIds = Future.Get();
 		if (MemberProductUserIds->Num() != ResolvedAccountIds.Num())
 		{
-			// Todo: Errors
+			UE_LOG(LogTemp, Warning, TEXT("[FLobbyDetailsEOS::GetLobbySnapshot] ResolveAccountIds Failed: Expected Count[%d], Received Count[%d], Lobby[%s]"),
+				MemberProductUserIds->Num(), ResolvedAccountIds.Num(), UTF8_TO_TCHAR(StrongThis->GetInfo()->GetLobbyIdEOS()));
 			Promise.EmplaceValue(Errors::Unknown());
 			return;
 		}
@@ -355,7 +384,8 @@ TFuture<TDefaultErrorResultInternal<TSharedRef<FClientLobbySnapshot>>> FLobbyDet
 				EOS_EResult EOSResult = EOS_LobbyDetails_CopyAttributeByIndex(StrongThis->LobbyDetailsHandle, &CopyAttributeByIndexOptions, &LobbyAttribute);
 				if (EOSResult != EOS_EResult::EOS_Success)
 				{
-					// todo: errors
+					UE_LOG(LogTemp, Warning, TEXT("[FLobbyDetailsEOS::GetLobbySnapshot] EOS_LobbyDetails_CopyAttributeByIndex Failed: Lobby[%s], Result[%s]"),
+						UTF8_TO_TCHAR(StrongThis->GetInfo()->GetLobbyIdEOS()), *LexToString(EOSResult));
 					Promise.EmplaceValue(Errors::FromEOSResult(EOSResult));
 					return;
 				}
@@ -365,9 +395,12 @@ TFuture<TDefaultErrorResultInternal<TSharedRef<FClientLobbySnapshot>>> FLobbyDet
 			}
 		}
 
+		UE_LOG(LogTemp, VeryVerbose, TEXT("[FLobbyDetailsEOS::GetLobbySnapshot] Succeeded: Lobby[%s]"),
+			UTF8_TO_TCHAR(StrongThis->GetInfo()->GetLobbyIdEOS()));
 		Promise.EmplaceValue(MoveTemp(ClientLobbySnapshot));
 	});
 
+
 	return Future;
 }
 
@@ -406,6 +439,8 @@ TDefaultErrorResultInternal<TSharedRef<FClientLobbyMemberSnapshot>> FLobbyDetail
 			EOS_EResult EOSResult = EOS_LobbyDetails_CopyMemberAttributeByIndex(LobbyDetailsHandle, &CopyMemberAttributeByIndexOptions, &LobbyAttribute);
 			if (EOSResult != EOS_EResult::EOS_Success)
 			{
+				UE_LOG(LogTemp, Warning, TEXT("[FLobbyDetailsEOS::GetLobbyMemberSnapshot] EOS_LobbyDetails_CopyMemberAttributeByIndex Failed: Member[%s], Lobby[%s], Result[%s]"),
+					*ToLogString(MemberAccountId), UTF8_TO_TCHAR(GetInfo()->GetLobbyIdEOS()), *LexToString(EOSResult));
 				return TDefaultErrorResultInternal<TSharedRef<FClientLobbyMemberSnapshot>>(Errors::FromEOSResult(EOSResult));
 			}
 
@@ -414,6 +449,8 @@ TDefaultErrorResultInternal<TSharedRef<FClientLobbyMemberSnapshot>> FLobbyDetail
 		}
 	}
 
+	UE_LOG(LogTemp, VeryVerbose, TEXT("[FLobbyDetailsEOS::GetLobbyMemberSnapshot] Succeeded: Member[%s], Lobby[%s]"),
+		*ToLogString(MemberAccountId), UTF8_TO_TCHAR(GetInfo()->GetLobbyIdEOS()));
 	return TDefaultErrorResultInternal<TSharedRef<FClientLobbyMemberSnapshot>>(MoveTemp(LobbyMemberSnapshot));
 }
 
@@ -432,29 +469,31 @@ TFuture<EOS_EResult> FLobbyDetailsEOS::ApplyLobbyDataUpdateFromLocalChanges(
 	EOS_Lobby_UpdateLobbyModificationOptions ModificationOptions = {};
 	ModificationOptions.ApiVersion = EOS_LOBBY_UPDATELOBBYMODIFICATION_API_LATEST;
 	ModificationOptions.LocalUserId = GetProductUserIdChecked(LocalAccountId);
-	ModificationOptions.LobbyId = GetInfo()->GetLobbyId();
+	ModificationOptions.LobbyId = GetInfo()->GetLobbyIdEOS();
 	static_assert(EOS_LOBBY_UPDATELOBBYMODIFICATION_API_LATEST == 1, "EOS_Lobby_UpdateLobbyModificationOptions updated, check new fields");
 
-	EOS_EResult EOSResultCode = EOS_Lobby_UpdateLobbyModification(Prerequisites->LobbyInterfaceHandle, &ModificationOptions, &LobbyModificationHandle);
-	if (EOSResultCode != EOS_EResult::EOS_Success)
+	EOS_EResult EOSResult = EOS_Lobby_UpdateLobbyModification(Prerequisites->LobbyInterfaceHandle, &ModificationOptions, &LobbyModificationHandle);
+	if (EOSResult != EOS_EResult::EOS_Success)
 	{
-		// Todo: Errors
-		return MakeFulfilledPromise<EOS_EResult>(EOSResultCode).GetFuture();
+		UE_LOG(LogTemp, Warning, TEXT("[FLobbyDetailsEOS::ApplyLobbyDataUpdateFromLocalChanges] EOS_Lobby_UpdateLobbyModification Failed: User[%s], Lobby[%s], Result[%s]"),
+			*ToLogString(LocalAccountId), UTF8_TO_TCHAR(GetInfo()->GetLobbyIdEOS()), *LexToString(EOSResult));
+		return MakeFulfilledPromise<EOS_EResult>(EOSResult).GetFuture();
 	}
 
+	// Set lobby join policy.
 	if (Changes.JoinPolicy)
 	{
-		// Set lobby join policy.
 		EOS_LobbyModification_SetPermissionLevelOptions SetPermissionOptions = {};
 		SetPermissionOptions.ApiVersion = EOS_LOBBYMODIFICATION_SETPERMISSIONLEVEL_API_LATEST;
 		SetPermissionOptions.PermissionLevel = TranslateJoinPolicy(*Changes.JoinPolicy);
 		static_assert(EOS_LOBBYMODIFICATION_SETPERMISSIONLEVEL_API_LATEST == 1, "EOS_LobbyModification_SetPermissionLevelOptions updated, check new fields");
 
-		EOSResultCode = EOS_LobbyModification_SetPermissionLevel(LobbyModificationHandle, &SetPermissionOptions);
-		if (EOSResultCode != EOS_EResult::EOS_Success)
+		EOSResult = EOS_LobbyModification_SetPermissionLevel(LobbyModificationHandle, &SetPermissionOptions);
+		if (EOSResult != EOS_EResult::EOS_Success)
 		{
-			// Todo: Errors
-			return MakeFulfilledPromise<EOS_EResult>(EOSResultCode).GetFuture();
+			UE_LOG(LogTemp, Warning, TEXT("[FLobbyDetailsEOS::ApplyLobbyDataUpdateFromLocalChanges] EOS_LobbyModification_SetPermissionLevel Failed: User[%s], Lobby[%s], Result[%s]"),
+				*ToLogString(LocalAccountId), UTF8_TO_TCHAR(GetInfo()->GetLobbyIdEOS()), *LexToString(EOSResult));
+			return MakeFulfilledPromise<EOS_EResult>(EOSResult).GetFuture();
 		}
 	}
 
@@ -469,11 +508,12 @@ TFuture<EOS_EResult> FLobbyDetailsEOS::ApplyLobbyDataUpdateFromLocalChanges(
 		AddAttributeOptions.Visibility = EOS_ELobbyAttributeVisibility::EOS_LAT_PUBLIC; // todo - get from schema
 		static_assert(EOS_LOBBYMODIFICATION_ADDATTRIBUTE_API_LATEST == 1, "EOS_LobbyModification_AddAttributeOptions updated, check new fields");
 
-		EOSResultCode = EOS_LobbyModification_AddAttribute(LobbyModificationHandle, &AddAttributeOptions);
-		if (EOSResultCode != EOS_EResult::EOS_Success)
+		EOSResult = EOS_LobbyModification_AddAttribute(LobbyModificationHandle, &AddAttributeOptions);
+		if (EOSResult != EOS_EResult::EOS_Success)
 		{
-			// Todo: Errors
-			return MakeFulfilledPromise<EOS_EResult>(EOSResultCode).GetFuture();
+			UE_LOG(LogTemp, Warning, TEXT("[FLobbyDetailsEOS::ApplyLobbyDataUpdateFromLocalChanges] EOS_LobbyModification_AddAttribute Failed: User[%s], Lobby[%s], Key[%s], Value[%s], Result[%s]"),
+				*ToLogString(LocalAccountId), UTF8_TO_TCHAR(GetInfo()->GetLobbyIdEOS()), *MutatedAttribute.Key.ToString().ToLower(), *ToLogString(MutatedAttribute.Value), *LexToString(EOSResult));
+			return MakeFulfilledPromise<EOS_EResult>(EOSResult).GetFuture();
 		}
 	}
 
@@ -487,11 +527,12 @@ TFuture<EOS_EResult> FLobbyDetailsEOS::ApplyLobbyDataUpdateFromLocalChanges(
 		RemoveAttributeOptions.Key = KeyConverter.Get();
 		static_assert(EOS_LOBBYMODIFICATION_REMOVEATTRIBUTE_API_LATEST == 1, "EOS_LobbyModification_RemoveAttributeOptions updated, check new fields");
 		
-		EOSResultCode = EOS_LobbyModification_RemoveAttribute(LobbyModificationHandle, &RemoveAttributeOptions);
-		if (EOSResultCode != EOS_EResult::EOS_Success)
+		EOSResult = EOS_LobbyModification_RemoveAttribute(LobbyModificationHandle, &RemoveAttributeOptions);
+		if (EOSResult != EOS_EResult::EOS_Success)
 		{
-			// Todo: Errors
-			return MakeFulfilledPromise<EOS_EResult>(EOSResultCode).GetFuture();
+			UE_LOG(LogTemp, Warning, TEXT("[FLobbyDetailsEOS::ApplyLobbyDataUpdateFromLocalChanges] EOS_LobbyModification_RemoveAttribute Failed: User[%s], Lobby[%s], Key[%s], Result[%s]"),
+				*ToLogString(LocalAccountId), UTF8_TO_TCHAR(GetInfo()->GetLobbyIdEOS()), UTF8_TO_TCHAR(KeyConverter.Get()), *LexToString(EOSResult));
+			return MakeFulfilledPromise<EOS_EResult>(EOSResult).GetFuture();
 		}
 	}
 
@@ -505,9 +546,20 @@ TFuture<EOS_EResult> FLobbyDetailsEOS::ApplyLobbyDataUpdateFromLocalChanges(
 	static_assert(EOS_LOBBY_UPDATELOBBY_API_LATEST == 1, "EOS_Lobby_UpdateLobbyOptions updated, check new fields");
 
 	EOS_Async(EOS_Lobby_UpdateLobby, Prerequisites->LobbyInterfaceHandle, UpdateLobbyOptions,
-	[Promise = MoveTemp(Promise)](const EOS_Lobby_UpdateLobbyCallbackInfo* CallbackInfo) mutable
+	[LocalAccountId, LobbyInfo = GetInfo(), Promise = MoveTemp(Promise)](const EOS_Lobby_UpdateLobbyCallbackInfo* Data) mutable
 	{
-		Promise.EmplaceValue(CallbackInfo->ResultCode);
+		if (Data->ResultCode != EOS_EResult::EOS_Success && Data->ResultCode != EOS_EResult::EOS_NoChange)
+		{
+			UE_LOG(LogTemp, Warning, TEXT("[FLobbyDetailsEOS::ApplyLobbyDataUpdateFromLocalChanges] EOS_Lobby_UpdateLobby Failed: User[%s], Lobby[%s], Result[%s]"),
+				*ToLogString(LocalAccountId), UTF8_TO_TCHAR(LobbyInfo->GetLobbyIdEOS()), *LexToString(Data->ResultCode));
+		}
+		else
+		{
+			UE_LOG(LogTemp, VeryVerbose, TEXT("[FLobbyDetailsEOS::ApplyLobbyDataUpdateFromLocalChanges] Succeeded: User[%s], Lobby[%s], Result[%s]"),
+				*ToLogString(LocalAccountId), UTF8_TO_TCHAR(LobbyInfo->GetLobbyIdEOS()), *LexToString(Data->ResultCode));
+		}
+
+		Promise.EmplaceValue(Data->ResultCode);
 	});
 	return Future;
 }
@@ -527,14 +579,15 @@ TFuture<EOS_EResult> FLobbyDetailsEOS::ApplyLobbyMemberDataUpdateFromLocalChange
 	EOS_Lobby_UpdateLobbyModificationOptions ModificationOptions = {};
 	ModificationOptions.ApiVersion = EOS_LOBBY_UPDATELOBBYMODIFICATION_API_LATEST;
 	ModificationOptions.LocalUserId = GetProductUserIdChecked(LocalAccountId);
-	ModificationOptions.LobbyId = GetInfo()->GetLobbyId();
+	ModificationOptions.LobbyId = GetInfo()->GetLobbyIdEOS();
 	static_assert(EOS_LOBBY_UPDATELOBBYMODIFICATION_API_LATEST == 1, "EOS_Lobby_UpdateLobbyModificationOptions updated, check new fields");
 
-	EOS_EResult EOSResultCode = EOS_Lobby_UpdateLobbyModification(Prerequisites->LobbyInterfaceHandle, &ModificationOptions, &LobbyModificationHandle);
-	if (EOSResultCode != EOS_EResult::EOS_Success)
+	EOS_EResult EOSResult = EOS_Lobby_UpdateLobbyModification(Prerequisites->LobbyInterfaceHandle, &ModificationOptions, &LobbyModificationHandle);
+	if (EOSResult != EOS_EResult::EOS_Success)
 	{
-		// Todo: Errors
-		return MakeFulfilledPromise<EOS_EResult>(EOSResultCode).GetFuture();
+		UE_LOG(LogTemp, Warning, TEXT("[FLobbyDetailsEOS::ApplyLobbyMemberDataUpdateFromLocalChanges] EOS_Lobby_UpdateLobbyModification Failed: User[%s], Lobby[%s], Result[%s]"),
+			*ToLogString(LocalAccountId), UTF8_TO_TCHAR(GetInfo()->GetLobbyIdEOS()), *LexToString(EOSResult));
+		return MakeFulfilledPromise<EOS_EResult>(EOSResult).GetFuture();
 	}
 
 	// Add member attributes.
@@ -548,11 +601,12 @@ TFuture<EOS_EResult> FLobbyDetailsEOS::ApplyLobbyMemberDataUpdateFromLocalChange
 		AddMemberAttributeOptions.Visibility = EOS_ELobbyAttributeVisibility::EOS_LAT_PUBLIC; // todo - get from schema
 		static_assert(EOS_LOBBYMODIFICATION_ADDMEMBERATTRIBUTE_API_LATEST == 1, "EOS_LobbyModification_AddMemberAttributeOptions updated, check new fields");
 		
-		EOSResultCode = EOS_LobbyModification_AddMemberAttribute(LobbyModificationHandle, &AddMemberAttributeOptions);
-		if (EOSResultCode != EOS_EResult::EOS_Success)
+		EOSResult = EOS_LobbyModification_AddMemberAttribute(LobbyModificationHandle, &AddMemberAttributeOptions);
+		if (EOSResult != EOS_EResult::EOS_Success)
 		{
-			// Todo: Errors
-			return MakeFulfilledPromise<EOS_EResult>(EOSResultCode).GetFuture();
+			UE_LOG(LogTemp, Warning, TEXT("[FLobbyDetailsEOS::ApplyLobbyMemberDataUpdateFromLocalChanges] EOS_LobbyModification_AddMemberAttribute Failed: User[%s], Lobby[%s], Key[%s], Value[%s], Result[%s]"),
+				*ToLogString(LocalAccountId), UTF8_TO_TCHAR(GetInfo()->GetLobbyIdEOS()), *MutatedAttribute.Key.ToString().ToLower(), *ToLogString(MutatedAttribute.Value), *LexToString(EOSResult));
+			return MakeFulfilledPromise<EOS_EResult>(EOSResult).GetFuture();
 		}
 	}
 
@@ -566,11 +620,12 @@ TFuture<EOS_EResult> FLobbyDetailsEOS::ApplyLobbyMemberDataUpdateFromLocalChange
 		RemoveMemberAttributeOptions.Key = KeyConverter.Get();
 		static_assert(EOS_LOBBYMODIFICATION_REMOVEMEMBERATTRIBUTE_API_LATEST == 1, "EOS_LobbyModification_RemoveMemberAttributeOptions updated, check new fields");
 		
-		EOSResultCode = EOS_LobbyModification_RemoveMemberAttribute(LobbyModificationHandle, &RemoveMemberAttributeOptions);
-		if (EOSResultCode != EOS_EResult::EOS_Success)
+		EOSResult = EOS_LobbyModification_RemoveMemberAttribute(LobbyModificationHandle, &RemoveMemberAttributeOptions);
+		if (EOSResult != EOS_EResult::EOS_Success)
 		{
-			// Todo: Errors
-			return MakeFulfilledPromise<EOS_EResult>(EOSResultCode).GetFuture();
+			UE_LOG(LogTemp, Warning, TEXT("[FLobbyDetailsEOS::ApplyLobbyMemberDataUpdateFromLocalChanges] EOS_LobbyModification_RemoveMemberAttribute Failed: User[%s], Lobby[%s], Key[%s], Result[%s]"),
+				*ToLogString(LocalAccountId), UTF8_TO_TCHAR(GetInfo()->GetLobbyIdEOS()), UTF8_TO_TCHAR(KeyConverter.Get()), *LexToString(EOSResult));
+			return MakeFulfilledPromise<EOS_EResult>(EOSResult).GetFuture();
 		}
 	}
 
@@ -584,9 +639,20 @@ TFuture<EOS_EResult> FLobbyDetailsEOS::ApplyLobbyMemberDataUpdateFromLocalChange
 	TFuture<EOS_EResult> Future = Promise.GetFuture();
 
 	EOS_Async(EOS_Lobby_UpdateLobby, Prerequisites->LobbyInterfaceHandle, UpdateLobbyOptions,
-	[Promise = MoveTemp(Promise)](const EOS_Lobby_UpdateLobbyCallbackInfo* CallbackInfo) mutable
+	[LocalAccountId, LobbyInfo = GetInfo(), Promise = MoveTemp(Promise)](const EOS_Lobby_UpdateLobbyCallbackInfo* Data) mutable
 	{
-		Promise.EmplaceValue(CallbackInfo->ResultCode);
+		if (Data->ResultCode != EOS_EResult::EOS_Success && Data->ResultCode != EOS_EResult::EOS_NoChange)
+		{
+			UE_LOG(LogTemp, Warning, TEXT("[FLobbyDetailsEOS::ApplyLobbyMemberDataUpdateFromLocalChanges] EOS_Lobby_UpdateLobby Failed: User[%s], Lobby[%s], Result[%s]"),
+				*ToLogString(LocalAccountId), UTF8_TO_TCHAR(LobbyInfo->GetLobbyIdEOS()), *LexToString(Data->ResultCode));
+		}
+		else
+		{
+			UE_LOG(LogTemp, VeryVerbose, TEXT("[FLobbyDetailsEOS::ApplyLobbyMemberDataUpdateFromLocalChanges] Succeeded: User[%s], Lobby[%s], Result[%s]"),
+				*ToLogString(LocalAccountId), UTF8_TO_TCHAR(LobbyInfo->GetLobbyIdEOS()), *LexToString(Data->ResultCode));
+		}
+
+		Promise.EmplaceValue(Data->ResultCode);
 	});
 
 	return Future;
@@ -616,6 +682,8 @@ TDefaultErrorResultInternal<TSharedRef<FLobbyDetailsInfoEOS>> FLobbyDetailsInfoE
 	EOS_EResult EOSResult = EOS_LobbyDetails_CopyInfo(LobbyDetailsHandle, &CopyInfoOptions, &LobbyDetailsInfo);
 	if (EOSResult != EOS_EResult::EOS_Success)
 	{
+		UE_LOG(LogTemp, Warning, TEXT("[FLobbyDetailsInfoEOS] EOS_LobbyDetails_CopyInfo Failed: Result[%s]"),
+			*LexToString(EOSResult));
 		return TDefaultErrorResultInternal<TSharedRef<FLobbyDetailsInfoEOS>>(Errors::FromEOSResult(EOSResult));
 	}
 
@@ -623,14 +691,15 @@ TDefaultErrorResultInternal<TSharedRef<FLobbyDetailsInfoEOS>> FLobbyDetailsInfoE
 }
 
 FLobbyDetailsInfoEOS::FLobbyDetailsInfoEOS(FLobbyDetailsInfoPtr&& InLobbyDetailsInfo)
-	: LobbyDetailsInfo(MoveTempIfPossible(InLobbyDetailsInfo))
+	: LobbyDetailsInfo(MoveTemp(InLobbyDetailsInfo))
 {
 	const FLobbyBucketIdTranslator<ELobbyTranslationType::FromService> BucketTranslator(LobbyDetailsInfo->BucketId);
 	BucketId = BucketTranslator.GetBucketId();
 
 	if (!BucketId.IsValid())
 	{
-		UE_LOG(LogTemp, Warning, TEXT("[FLobbyDetailsInfoEOS] Failed to parse lobby bucket id. Lobby: %s, Bucket: %s"), UTF8_TO_TCHAR(LobbyDetailsInfo->LobbyId), UTF8_TO_TCHAR(LobbyDetailsInfo->BucketId));
+		UE_LOG(LogTemp, Warning, TEXT("[FLobbyDetailsInfoEOS] Failed: Unable to parse lobby bucket id. Lobby[%s], Bucket[%s]"),
+			UTF8_TO_TCHAR(LobbyDetailsInfo->LobbyId), UTF8_TO_TCHAR(LobbyDetailsInfo->BucketId));
 	}
 }
 
@@ -684,7 +753,7 @@ FLobbyDataEOS::FLobbyDataEOS(
 	: ClientLobbyData(ClientLobbyData)
 	, LobbyDetailsInfo(LobbyDetailsInfo)
 	, UnregisterFn(MoveTemp(UnregisterFn))
-	, LobbyId(Private::TranslateLobbyId(LobbyDetailsInfo->GetLobbyId()))
+	, LobbyIdString(UTF8_TO_TCHAR(LobbyDetailsInfo->GetLobbyIdEOS()))
 {
 }
 
@@ -697,24 +766,25 @@ TFuture<TDefaultErrorResultInternal<TSharedRef<FLobbyDataEOS>>> FLobbyDataEOS::C
 	TFuture<TDefaultErrorResultInternal<TSharedRef<FLobbyDataEOS>>> Future = Promise.GetFuture();
 	
 	LobbyDetails->GetLobbySnapshot()
-	.Then(
+	.Next(
 	[
 		Promise = MoveTemp(Promise),
 		LobbyId,
 		LobbyDetails,
 		UnregisterFn = MoveTemp(UnregisterFn)
 	]
-	(TFuture<TDefaultErrorResultInternal<TSharedRef<FClientLobbySnapshot>>>&& Future) mutable
+	(TDefaultErrorResultInternal<TSharedRef<FClientLobbySnapshot>>&& Result) mutable
 	{
-		if (Future.Get().IsError())
+		if (Result.IsError())
 		{
-			// todo: errors.
-			Promise.EmplaceValue(MoveTemp(Future.Get().GetErrorValue()));
+			UE_LOG(LogTemp, Warning, TEXT("[FLobbyDataEOS::Create] FLobbyDetailsEOS::GetLobbySnapshot Failed: Lobby[%s], Result[%s]"),
+				*ToLogString(LobbyId), *Result.GetErrorValue().GetLogString());
+			Promise.EmplaceValue(MoveTemp(Result.GetErrorValue()));
 			return;
 		}
 
-		TSharedRef<FClientLobbySnapshot> LobbySnapshot = Future.Get().GetOkValue();
-		TSharedRef<FClientLobbyData> LobbyData = MakeShared<FClientLobbyData>(LobbyId);
+		TSharedRef<FClientLobbySnapshot> LobbySnapshot = Result.GetOkValue();
+		TSharedRef<FClientLobbyData> NewClientLobbyData = MakeShared<FClientLobbyData>(LobbyId);
 
 		// Fetch member data and apply them to the lobby.
 		TMap<FAccountId, TSharedRef<FClientLobbyMemberSnapshot>> MemberSnapshots;
@@ -723,7 +793,8 @@ TFuture<TDefaultErrorResultInternal<TSharedRef<FLobbyDataEOS>>> FLobbyDataEOS::C
 			TDefaultErrorResultInternal<TSharedRef<FClientLobbyMemberSnapshot>> LobbyMemberSnapshotResult = LobbyDetails->GetLobbyMemberSnapshot(MemberAccountId);
 			if (LobbyMemberSnapshotResult.IsError())
 			{
-				// todo: errors.
+				UE_LOG(LogTemp, Warning, TEXT("[FLobbyDataEOS::Create] FLobbyDetailsEOS::GetLobbyMemberSnapshot Failed: Lobby[%s], User[%s], Result[%s]"),
+					*ToLogString(LobbyId), *ToLogString(MemberAccountId), *Result.GetErrorValue().GetLogString());
 				Promise.EmplaceValue(MoveTemp(LobbyMemberSnapshotResult.GetErrorValue()));
 				return;
 			}
@@ -731,10 +802,12 @@ TFuture<TDefaultErrorResultInternal<TSharedRef<FLobbyDataEOS>>> FLobbyDataEOS::C
 			MemberSnapshots.Emplace(MemberAccountId, MoveTemp(LobbyMemberSnapshotResult.GetOkValue()));
 		}
 
-		LobbyData->ApplyLobbyUpdateFromServiceSnapshot(MoveTemp(*LobbySnapshot), MoveTemp(MemberSnapshots));
+		NewClientLobbyData->ApplyLobbyUpdateFromServiceSnapshot(MoveTemp(*LobbySnapshot), MoveTemp(MemberSnapshots));
 
+		UE_LOG(LogTemp, VeryVerbose, TEXT("[FLobbyDataEOS::Create] Succeeded: Lobby[%s]"),
+			UTF8_TO_TCHAR(LobbyDetails->GetInfo()->GetLobbyIdEOS()));
 		Promise.EmplaceValue(MakeShared<FLobbyDataEOS>(
-			LobbyData,
+			NewClientLobbyData,
 			LobbyDetails->GetInfo(),
 			MoveTemp(UnregisterFn)));
 	});
@@ -747,9 +820,9 @@ FLobbyDataRegistryEOS::FLobbyDataRegistryEOS(const TSharedRef<FLobbyPrerequisite
 {
 }
 
-TSharedPtr<FLobbyDataEOS> FLobbyDataRegistryEOS::Find(EOS_LobbyId EOSLobbyId) const
+TSharedPtr<FLobbyDataEOS> FLobbyDataRegistryEOS::Find(EOS_LobbyId LobbyIdEOS) const
 {
-	const TWeakPtr<FLobbyDataEOS>* Result = LobbyIdIndex.Find(Private::TranslateLobbyId(EOSLobbyId));
+	const TWeakPtr<FLobbyDataEOS>* Result = LobbyIdIndex.Find(UTF8_TO_TCHAR(LobbyIdEOS));
 	return Result ? Result->Pin() : TSharedPtr<FLobbyDataEOS>();
 }
 
@@ -761,8 +834,12 @@ TSharedPtr<FLobbyDataEOS> FLobbyDataRegistryEOS::Find(FLobbyId LobbyId) const
 
 TFuture<TDefaultErrorResultInternal<TSharedRef<FLobbyDataEOS>>> FLobbyDataRegistryEOS::FindOrCreateFromLobbyDetails(FAccountId LocalAccountId, const TSharedRef<FLobbyDetailsEOS>& LobbyDetails)
 {
-	if (TSharedPtr<FLobbyDataEOS> FindResult = Find(LobbyDetails->GetInfo()->GetLobbyId()))
+	// Check if lobby data already exists.
+	if (TSharedPtr<FLobbyDataEOS> FindResult = Find(LobbyDetails->GetInfo()->GetLobbyIdEOS()))
 	{
+		UE_LOG(LogTemp, VeryVerbose, TEXT("[FLobbyDataRegistryEOS::FindOrCreateFromLobbyDetails] Succeeded: Lobby data already exists. User[%s], Lobby[%s]"),
+			*ToLogString(LocalAccountId), *FindResult->GetLobbyIdString());
+
 		FindResult->AddUserLobbyDetails(LocalAccountId, LobbyDetails);
 		return MakeFulfilledPromise<TDefaultErrorResultInternal<TSharedRef<FLobbyDataEOS>>>(FindResult.ToSharedRef()).GetFuture();
 	}
@@ -770,20 +847,36 @@ TFuture<TDefaultErrorResultInternal<TSharedRef<FLobbyDataEOS>>> FLobbyDataRegist
 	TPromise<TDefaultErrorResultInternal<TSharedRef<FLobbyDataEOS>>> Promise;
 	TFuture<TDefaultErrorResultInternal<TSharedRef<FLobbyDataEOS>>> Future = Promise.GetFuture();
 
+	// Create new lobby id object.
 	const FLobbyId LobbyId = FLobbyId(EOnlineServices::Epic, NextHandleIndex++);
+
+	// Create lobby data from lobby details.
 	FLobbyDataEOS::Create(LobbyId, LobbyDetails, MakeUnregisterFn())
-	.Then([WeakThis = AsWeak(), Promise = MoveTemp(Promise), LocalAccountId, LobbyDetails](TFuture<TDefaultErrorResultInternal<TSharedRef<FLobbyDataEOS>>>&& Future) mutable
+	.Next([WeakThis = AsWeak(), Promise = MoveTemp(Promise), LocalAccountId, LobbyDetails](TDefaultErrorResultInternal<TSharedRef<FLobbyDataEOS>>&& Result) mutable
 	{
-		if (TSharedPtr<FLobbyDataRegistryEOS> StrongThis = WeakThis.Pin())
+		TSharedPtr<FLobbyDataRegistryEOS> StrongThis = WeakThis.Pin();
+		if (!StrongThis.IsValid())
 		{
-			if (Future.Get().IsOk())
-			{
-				StrongThis->Register(Future.Get().GetOkValue());
-			}
+			UE_LOG(LogTemp, Warning, TEXT("[FLobbyDataRegistryEOS::FindOrCreateFromLobbyDetails] Failed: FLobbyDataRegistryEOS has been destroyed. User[%s], Lobby[%s], Result[%s]"),
+				*ToLogString(LocalAccountId), UTF8_TO_TCHAR(LobbyDetails->GetInfo()->GetLobbyIdEOS()), *Result.GetErrorValue().GetLogString());
+			Promise.EmplaceValue(Errors::NotFound());
+			return;
+		}
+
+		if (Result.IsError())
+		{
+			UE_LOG(LogTemp, Warning, TEXT("[FLobbyDataRegistryEOS::FindOrCreateFromLobbyDetails] FLobbyDataEOS::Create Failed: User[%s], Lobby[%s], Result[%s]"),
+				*ToLogString(LocalAccountId), UTF8_TO_TCHAR(LobbyDetails->GetInfo()->GetLobbyIdEOS()), *Result.GetErrorValue().GetLogString());
+			Promise.EmplaceValue(MoveTemp(Result.GetErrorValue()));
+			return;
 		}
 
-		Future.Get().GetOkValue()->AddUserLobbyDetails(LocalAccountId, LobbyDetails);
-		Promise.EmplaceValue(MoveTempIfPossible(Future.Get()));
+		UE_LOG(LogTemp, VeryVerbose, TEXT("[FLobbyDataRegistryEOS::FindOrCreateFromLobbyDetails] Succeeded: Created new lobby data. User[%s], Lobby[%s]"),
+			*ToLogString(LocalAccountId), *Result.GetOkValue()->GetLobbyIdString());
+
+		StrongThis->Register(Result.GetOkValue());
+		Result.GetOkValue()->AddUserLobbyDetails(LocalAccountId, LobbyDetails);
+		Promise.EmplaceValue(MoveTemp(Result.GetOkValue()));
 	});
 
 	return Future;
@@ -791,7 +884,7 @@ TFuture<TDefaultErrorResultInternal<TSharedRef<FLobbyDataEOS>>> FLobbyDataRegist
 
 void FLobbyDataRegistryEOS::Register(const TSharedRef<FLobbyDataEOS>& LobbyIdHandleData)
 {
-	LobbyIdIndex.Add(LobbyIdHandleData->GetLobbyId(), LobbyIdHandleData);
+	LobbyIdIndex.Add(LobbyIdHandleData->GetLobbyIdEOS(), LobbyIdHandleData);
 	LobbyIdHandleIndex.Add(LobbyIdHandleData->GetLobbyIdHandle(), LobbyIdHandleData);
 }
 
@@ -799,7 +892,7 @@ void FLobbyDataRegistryEOS::Unregister(FLobbyId LobbyId)
 {
 	if (TSharedPtr<FLobbyDataEOS> HandleData = Find(LobbyId))
 	{
-		LobbyIdIndex.Remove(HandleData->GetLobbyId());
+		LobbyIdIndex.Remove(HandleData->GetLobbyIdEOS());
 		LobbyIdHandleIndex.Remove(HandleData->GetLobbyIdHandle());
 	}
 }
@@ -825,6 +918,9 @@ TFuture<TDefaultErrorResultInternal<TSharedRef<FLobbyInviteDataEOS>>> FLobbyInvi
 	TDefaultErrorResultInternal<TSharedRef<FLobbyDetailsEOS>> LobbyDetailsResult = FLobbyDetailsEOS::CreateFromInviteId(Prerequisites, LocalAccountId, InviteIdEOS);
 	if (LobbyDetailsResult.IsError())
 	{
+		UE_LOG(LogTemp, Warning, TEXT("[FLobbyInviteDataEOS::CreateFromInviteId] FLobbyDetailsEOS::CreateFromInviteId Failed: User[%s], InviteId[%s], Result[%s]"),
+			*ToLogString(LocalAccountId), UTF8_TO_TCHAR(InviteIdEOS), *LobbyDetailsResult.GetErrorValue().GetLogString());
+
 		return MakeFulfilledPromise<TDefaultErrorResultInternal<TSharedRef<FLobbyInviteDataEOS>>>(MoveTemp(LobbyDetailsResult.GetErrorValue())).GetFuture();
 	}
 
@@ -835,11 +931,14 @@ TFuture<TDefaultErrorResultInternal<TSharedRef<FLobbyInviteDataEOS>>> FLobbyInvi
 	// Search for existing lobby data so that the LobbyId will match.
 	TSharedRef<FLobbyInviteIdEOS> InviteId = MakeShared<FLobbyInviteIdEOS>(InviteIdEOS);
 	LobbyDataRegistry->FindOrCreateFromLobbyDetails(LocalAccountId, LobbyDetails)
-	.Then([Promise = MoveTemp(Promise), InviteId, LocalAccountId, Sender, LobbyDetails](TFuture<TDefaultErrorResultInternal<TSharedRef<FLobbyDataEOS>>>&& Future) mutable
+	.Next([Promise = MoveTemp(Promise), InviteId, LocalAccountId, Sender, LobbyDetails](TDefaultErrorResultInternal<TSharedRef<FLobbyDataEOS>>&& Result) mutable
 	{
-		if (Future.Get().IsError())
+		if (Result.IsError())
 		{
-			Promise.EmplaceValue(MoveTemp(Future.Get().GetErrorValue()));
+			UE_LOG(LogTemp, Warning, TEXT("[FLobbyInviteDataEOS::CreateFromInviteId] FLobbyDataEOS::FindOrCreateFromLobbyDetails Failed: User[%s], InviteId[%s], Lobby[%s], Result[%s]"),
+				*ToLogString(LocalAccountId), UTF8_TO_TCHAR(InviteId->Get()), UTF8_TO_TCHAR(LobbyDetails->GetInfo()->GetLobbyIdEOS()), *Result.GetErrorValue().GetLogString());
+
+			Promise.EmplaceValue(MoveTemp(Result.GetErrorValue()));
 			return;
 		}
 
@@ -847,12 +946,16 @@ TFuture<TDefaultErrorResultInternal<TSharedRef<FLobbyInviteDataEOS>>> FLobbyInvi
 		const FAccountId SenderAccountId = FindAccountId(Sender);
 		if (!SenderAccountId.IsValid())
 		{
-			// Todo: Errors.
-			Promise.EmplaceValue(MoveTemp(Future.Get().GetErrorValue()));
+			UE_LOG(LogTemp, Warning, TEXT("[FLobbyInviteDataEOS::CreateFromInviteId] Failed: Unable to find account id for sender. User[%s], InviteId[%s], Lobby[%s]"),
+				*ToLogString(LocalAccountId), UTF8_TO_TCHAR(InviteId->Get()), UTF8_TO_TCHAR(LobbyDetails->GetInfo()->GetLobbyIdEOS()));
+
+			Promise.EmplaceValue(Errors::InvalidState());
 			return;
 		}
 
-		Promise.EmplaceValue(MakeShared<FLobbyInviteDataEOS>(InviteId, LocalAccountId, SenderAccountId, LobbyDetails, Future.Get().GetOkValue()));
+		UE_LOG(LogTemp, VeryVerbose, TEXT("[FLobbyInviteDataEOS::CreateFromInviteId] Succeeded: User[%s], Sender[%s], InviteId[%s], Lobby[%s]"),
+			*ToLogString(LocalAccountId), *ToLogString(SenderAccountId), UTF8_TO_TCHAR(InviteId->Get()), *Result.GetOkValue()->GetLobbyIdString());
+		Promise.EmplaceValue(MakeShared<FLobbyInviteDataEOS>(InviteId, LocalAccountId, SenderAccountId, LobbyDetails, Result.GetOkValue()));
 	});
 
 	return Future;
@@ -869,7 +972,7 @@ FLobbyInviteDataEOS::FLobbyInviteDataEOS(
 	, Sender(Sender)
 	, LobbyDetails(LobbyDetails)
 	, LobbyData(LobbyData)
-	, InviteId(Private::TranslateLobbyInviteId(InviteIdEOS->Get()))
+	, InviteId(UTF8_TO_TCHAR(InviteIdEOS->Get()))
 {
 }
 
@@ -888,7 +991,8 @@ TFuture<TDefaultErrorResultInternal<TSharedRef<FLobbySearchEOS>>> FLobbySearchEO
 	EOS_EResult EOSResult = EOS_Lobby_CreateLobbySearch(Prerequisites->LobbyInterfaceHandle, &CreateLobbySearchOptions, &SearchHandle->Get());
 	if (EOSResult != EOS_EResult::EOS_Success)
 	{
-		// todo: errors
+		UE_LOG(LogTemp, Warning, TEXT("[FLobbySearchEOS::Create] EOS_Lobby_CreateLobbySearch Failed: User[%s], Result[%s]"),
+			*ToLogString(Params.LocalAccountId), *LexToString(EOSResult));
 		return MakeFulfilledPromise<TDefaultErrorResultInternal<TSharedRef<FLobbySearchEOS>>>(Errors::FromEOSResult(EOSResult)).GetFuture();
 	}
 
@@ -897,6 +1001,8 @@ TFuture<TDefaultErrorResultInternal<TSharedRef<FLobbySearchEOS>>> FLobbySearchEO
 		TSharedPtr<FLobbyDataEOS> LobbyData = LobbyRegistry->Find(*Params.LobbyId);
 		if (!LobbyData)
 		{
+			UE_LOG(LogTemp, Warning, TEXT("[FLobbySearchEOS::Create] Failed: Unable to find lobby data for lobby. User[%s], Lobby[%s]"),
+				*ToLogString(Params.LocalAccountId), *ToLogString(Params.LobbyId));
 			return MakeFulfilledPromise<TDefaultErrorResultInternal<TSharedRef<FLobbySearchEOS>>>(Errors::InvalidParams()).GetFuture();
 		}
 
@@ -908,7 +1014,8 @@ TFuture<TDefaultErrorResultInternal<TSharedRef<FLobbySearchEOS>>> FLobbySearchEO
 		EOSResult = EOS_LobbySearch_SetLobbyId(SearchHandle->Get(), &SetLobbyIdOptions);
 		if (EOSResult != EOS_EResult::EOS_Success)
 		{
-			// todo: errors
+			UE_LOG(LogTemp, Warning, TEXT("[FLobbySearchEOS::Create] EOS_LobbySearch_SetLobbyId Failed: User[%s], Lobby[%s], Result[%s]"),
+				*ToLogString(Params.LocalAccountId), *LobbyData->GetLobbyIdString(), *LexToString(EOSResult));
 			return MakeFulfilledPromise<TDefaultErrorResultInternal<TSharedRef<FLobbySearchEOS>>>(Errors::FromEOSResult(EOSResult)).GetFuture();
 		}
 	}
@@ -922,7 +1029,8 @@ TFuture<TDefaultErrorResultInternal<TSharedRef<FLobbySearchEOS>>> FLobbySearchEO
 		EOSResult = EOS_LobbySearch_SetTargetUserId(SearchHandle->Get(), &SetTargetUserIdOptions);
 		if (EOSResult != EOS_EResult::EOS_Success)
 		{
-			// todo: errors
+			UE_LOG(LogTemp, Warning, TEXT("[FLobbySearchEOS::Create] EOS_LobbySearch_SetTargetUserId Failed: User[%s], Target[%s], Result[%s]"),
+				*ToLogString(Params.LocalAccountId), *ToLogString(*Params.TargetUser), *LexToString(EOSResult));
 			return MakeFulfilledPromise<TDefaultErrorResultInternal<TSharedRef<FLobbySearchEOS>>>(Errors::FromEOSResult(EOSResult)).GetFuture();
 		}
 	}
@@ -948,7 +1056,8 @@ TFuture<TDefaultErrorResultInternal<TSharedRef<FLobbySearchEOS>>> FLobbySearchEO
 			EOSResult = EOS_LobbySearch_SetParameter(SearchHandle->Get(), &SetParameterOptions);
 			if (EOSResult != EOS_EResult::EOS_Success)
 			{
-				// todo: errors
+				UE_LOG(LogTemp, Warning, TEXT("[FLobbySearchEOS::Create] EOS_LobbySearch_SetParameter Failed: User[%s], Key[%s], Result[%s]"),
+					*ToLogString(Params.LocalAccountId), UTF8_TO_TCHAR(EOS_LOBBY_SEARCH_BUCKET_ID), *LexToString(EOSResult));
 				return MakeFulfilledPromise<TDefaultErrorResultInternal<TSharedRef<FLobbySearchEOS>>>(Errors::FromEOSResult(EOSResult)).GetFuture();
 			}
 		}
@@ -967,7 +1076,8 @@ TFuture<TDefaultErrorResultInternal<TSharedRef<FLobbySearchEOS>>> FLobbySearchEO
 			EOSResult = EOS_LobbySearch_SetParameter(SearchHandle->Get(), &SetParameterOptions);
 			if (EOSResult != EOS_EResult::EOS_Success)
 			{
-				// todo: errors
+				UE_LOG(LogTemp, Warning, TEXT("[FLobbySearchEOS::Create] EOS_LobbySearch_SetParameter Failed: User[%s], Key[%s], Value[%s], Result[%s]"),
+					*ToLogString(Params.LocalAccountId), *Filter.AttributeName.ToString().ToLower(), *ToLogString(Filter.ComparisonValue), *LexToString(EOSResult));
 				return MakeFulfilledPromise<TDefaultErrorResultInternal<TSharedRef<FLobbySearchEOS>>>(Errors::FromEOSResult(EOSResult)).GetFuture();
 			}
 		}
@@ -982,13 +1092,19 @@ TFuture<TDefaultErrorResultInternal<TSharedRef<FLobbySearchEOS>>> FLobbySearchEO
 	static_assert(EOS_LOBBYSEARCH_FIND_API_LATEST == 1, "EOS_LobbySearch_FindOptions updated, check new fields");
 
 	EOS_Async(EOS_LobbySearch_Find, SearchHandle->Get(), FindOptions,
-	[Promise = MoveTemp(Promise), Prerequisites, LobbyRegistry, LocalAccountId = Params.LocalAccountId, SearchHandle]
-	(const EOS_LobbySearch_FindCallbackInfo* CallbackInfo) mutable
+	[
+		Promise = MoveTemp(Promise),
+		Prerequisites,
+		LobbyRegistry,
+		LocalAccountId = Params.LocalAccountId,
+		SearchHandle
+	](const EOS_LobbySearch_FindCallbackInfo* Data) mutable
 	{
-		if (CallbackInfo->ResultCode != EOS_EResult::EOS_Success)
+		if (Data->ResultCode != EOS_EResult::EOS_Success)
 		{
-			// todo: errors
-			Promise.EmplaceValue(Errors::FromEOSResult(CallbackInfo->ResultCode));
+			UE_LOG(LogTemp, Warning, TEXT("[FLobbySearchEOS::Create] EOS_LobbySearch_Find Failed: User[%s], Result[%s]"),
+				*ToLogString(LocalAccountId), *LexToString(Data->ResultCode));
+			Promise.EmplaceValue(Errors::FromEOSResult(Data->ResultCode));
 			return;
 		}
 
@@ -1006,7 +1122,8 @@ TFuture<TDefaultErrorResultInternal<TSharedRef<FLobbySearchEOS>>> FLobbySearchEO
 				Prerequisites, LocalAccountId, SearchHandle->Get(), SearchResultIndex);
 			if (Result.IsError())
 			{
-				// todo: errors
+				UE_LOG(LogTemp, Warning, TEXT("[FLobbySearchEOS::Create] FLobbyDetailsEOS::CreateFromSearchResult Failed: User[%s], Result[%s]"),
+					*ToLogString(LocalAccountId), *Result.GetErrorValue().GetLogString());
 				Promise.EmplaceValue(MoveTemp(Result.GetErrorValue()));
 				return;
 			}
@@ -1015,24 +1132,24 @@ TFuture<TDefaultErrorResultInternal<TSharedRef<FLobbySearchEOS>>> FLobbySearchEO
 			ResolvedLobbyDetails.Add(ResolveLobbyDetailsPromise.GetFuture());
 
 			LobbyRegistry->FindOrCreateFromLobbyDetails(LocalAccountId, Result.GetOkValue())
-			.Then([ResolveLobbyDetailsPromise = MoveTemp(ResolveLobbyDetailsPromise)](TFuture<TDefaultErrorResultInternal<TSharedRef<FLobbyDataEOS>>>&& Future) mutable
+			.Next([ResolveLobbyDetailsPromise = MoveTemp(ResolveLobbyDetailsPromise)](TDefaultErrorResultInternal<TSharedRef<FLobbyDataEOS>>&& Result) mutable
 			{
-				ResolveLobbyDetailsPromise.EmplaceValue(MoveTempIfPossible(Future.Get()));
+				ResolveLobbyDetailsPromise.EmplaceValue(MoveTemp(Result));
 			});
 		}
 
 		WhenAll(MoveTemp(ResolvedLobbyDetails))
-		.Then([Promise = MoveTemp(Promise), SearchHandle](TFuture<TArray<TDefaultErrorResultInternal<TSharedRef<FLobbyDataEOS>>>>&& Future) mutable
+		.Next([LocalAccountId, Promise = MoveTemp(Promise), SearchHandle](TArray<TDefaultErrorResultInternal<TSharedRef<FLobbyDataEOS>>>&& ResultArray) mutable
 		{
-			TArray<TDefaultErrorResultInternal<TSharedRef<FLobbyDataEOS>>> Results(Future.Get());
 			TArray<TSharedRef<FLobbyDataEOS>> ResolvedResults;
-			ResolvedResults.Reserve(Results.Num());
+			ResolvedResults.Reserve(ResultArray.Num());
 
-			for (TDefaultErrorResultInternal<TSharedRef<FLobbyDataEOS>>& Result : Results)
+			for (TDefaultErrorResultInternal<TSharedRef<FLobbyDataEOS>>& Result : ResultArray)
 			{
 				if (Result.IsError())
 				{
-					// todo: errors
+					UE_LOG(LogTemp, Warning, TEXT("[FLobbySearchEOS::Create] FLobbyDataRegistryEOS::FindOrCreateFromLobbyDetails Failed: Unable to resolve search result. User[%s], Result[%s]"),
+						*ToLogString(LocalAccountId), *Result.GetErrorValue().GetLogString());
 					Promise.EmplaceValue(MoveTemp(Result.GetErrorValue()));
 					return;
 				}
@@ -1040,6 +1157,8 @@ TFuture<TDefaultErrorResultInternal<TSharedRef<FLobbySearchEOS>>> FLobbySearchEO
 				ResolvedResults.Add(MoveTemp(Result.GetOkValue()));
 			}
 
+			UE_LOG(LogTemp, VeryVerbose, TEXT("[FLobbySearchEOS::Create] Succeeded: User[%s], NumResults[%d]"),
+				*ToLogString(LocalAccountId), ResolvedResults.Num());
 			Promise.EmplaceValue(MakeShared<FLobbySearchEOS>(SearchHandle, MoveTemp(ResolvedResults)));
 		});
 	});
@@ -1073,7 +1192,7 @@ FLobbySearchEOS::FLobbySearchEOS(const TSharedRef<FSearchHandle>& SearchHandle,
 
 FString ToLogString(const FLobbyDataEOS& LobbyData)
 {
-	return FString::Printf(TEXT("[%s:%s]"), *ToLogString(LobbyData.GetLobbyIdHandle()), *LobbyData.GetLobbyId());
+	return FString::Printf(TEXT("[%s:%s]"), *ToLogString(LobbyData.GetLobbyIdHandle()), *LobbyData.GetLobbyIdString());
 }
 
 /* UE::Online */ }
\ No newline at end of file
diff --git a/Engine/Plugins/Online/Experimental/OnlineServicesEOSGS/Source/Private/Online/LobbiesEOSGSTypes.h b/Engine/Plugins/Online/Experimental/OnlineServicesEOSGS/Source/Private/Online/LobbiesEOSGSTypes.h
--- a/Engine/Plugins/Online/Experimental/OnlineServicesEOSGS/Source/Private/Online/LobbiesEOSGSTypes.h
+++ b/Engine/Plugins/Online/Experimental/OnlineServicesEOSGS/Source/Private/Online/LobbiesEOSGSTypes.h
@@ -40,6 +40,16 @@ public:
 	FLobbyBucketIdEOS& operator=(const FLobbyBucketIdEOS&) = default;
 	FLobbyBucketIdEOS& operator=(FLobbyBucketIdEOS&&) = default;
 
+	bool operator==(const FLobbyBucketIdEOS& Other) const
+	{
+		return ProductName == Other.ProductName && ProductVersion == Other.ProductVersion;
+	}
+
+	bool operator!=(const FLobbyBucketIdEOS& Other) const
+	{
+		return !(*this == Other);
+	}
+
 	FLobbyBucketIdEOS(FString ProductName, int32 ProductVersion);
 
 	const FString& GetProductName() const { return ProductName; }
@@ -51,6 +61,8 @@ private:
 	int32 ProductVersion;
 };
 
+FString ToLogString(const FLobbyBucketIdEOS& BucketId);
+
 //--------------------------------------------------------------------------------------------------
 // Translators
 //--------------------------------------------------------------------------------------------------
@@ -214,7 +226,11 @@ public:
 	static TDefaultErrorResultInternal<TSharedRef<FLobbyDetailsEOS>> CreateFromLobbyId(
 		const TSharedRef<FLobbyPrerequisitesEOS>& Prerequisites,
 		FAccountId LocalAccountId,
-		EOS_LobbyId LobbyId);
+		EOS_LobbyId LobbyIdEOS);
+	static TDefaultErrorResultInternal<TSharedRef<FLobbyDetailsEOS>> CreateFromLobbyId(
+		const TSharedRef<FLobbyPrerequisitesEOS>& Prerequisites,
+		FAccountId LocalAccountId,
+		const FString& LobbyId);
 	static TDefaultErrorResultInternal<TSharedRef<FLobbyDetailsEOS>> CreateFromInviteId(
 		const TSharedRef<FLobbyPrerequisitesEOS>& Prerequisites,
 		FAccountId LocalAccountId,
@@ -237,6 +253,11 @@ public:
 	ELobbyDetailsSource GetDetailsSource() const { return LobbyDetailsSource; }
 	FAccountId GetAssociatedUser() const { return AssociatedLocalUser; }
 
+	/**
+	 * Check whether this lobbies bucket is compatible with the currently running application.
+	 */
+	bool IsBucketCompatible() const;
+
 	/**
 	 * Retrieve a lobby data snapshot from the EOS lobby details object.
 	 * The lobby schema will be used to translate attribute data before returning.
@@ -288,12 +309,10 @@ public:
 
 	static TDefaultErrorResultInternal<TSharedRef<FLobbyDetailsInfoEOS>> Create(EOS_HLobbyDetails LobbyDetailsHandle);
 
-	EOS_LobbyId GetLobbyId() const { return LobbyDetailsInfo->LobbyId; }
+	EOS_LobbyId GetLobbyIdEOS() const { return LobbyDetailsInfo->LobbyId; }
 	int32 GetMaxMembers() const { return LobbyDetailsInfo->MaxMembers; }
 	EOS_ELobbyPermissionLevel GetPermissionLevel() const { return LobbyDetailsInfo->PermissionLevel; }
-
-	const FString& GetProductName() const { return BucketId.GetProductName(); }
-	int32 GetProductVersion() const { return BucketId.GetProductVersion(); }
+	const FLobbyBucketIdEOS& GetBucketId() const { return BucketId; }
 
 private:
 	struct FEOSLobbyDetailsInfoDeleter
@@ -331,8 +350,8 @@ public:
 
 	FLobbyId GetLobbyIdHandle() const { return ClientLobbyData->GetPublicData().LobbyId; }
 	const TSharedRef<FClientLobbyData>& GetClientLobbyData() const { return ClientLobbyData; }
-	EOS_LobbyId GetLobbyIdEOS() const { return LobbyDetailsInfo->GetLobbyId(); }
-	const FString& GetLobbyId() const { return LobbyId; }
+	EOS_LobbyId GetLobbyIdEOS() const { return LobbyDetailsInfo->GetLobbyIdEOS(); }
+	const FString& GetLobbyIdString() const { return LobbyIdString; }
 
 	void AddUserLobbyDetails(FAccountId LocalAccountId, const TSharedPtr<FLobbyDetailsEOS>& LobbyDetails);
 	TSharedPtr<FLobbyDetailsEOS> GetUserLobbyDetails(FAccountId LocalAccountId) const;
@@ -355,7 +374,7 @@ private:
 	TSharedRef<FClientLobbyData> ClientLobbyData;
 	TSharedRef<FLobbyDetailsInfoEOS> LobbyDetailsInfo;
 	FUnregisterFn UnregisterFn;
-	FString LobbyId;
+	FString LobbyIdString;
 	TMap<FAccountId, TSharedPtr<FLobbyDetailsEOS>> UserLobbyDetails;
 };
 
@@ -365,7 +384,7 @@ class FLobbyDataRegistryEOS : public TSharedFromThis<FLobbyDataRegistryEOS>
 public:
 	FLobbyDataRegistryEOS(const TSharedRef<FLobbyPrerequisitesEOS>& Prerequisites);
 
-	TSharedPtr<FLobbyDataEOS> Find(EOS_LobbyId EOSLobbyId) const;
+	TSharedPtr<FLobbyDataEOS> Find(EOS_LobbyId LobbyIdEOS) const;
 	TSharedPtr<FLobbyDataEOS> Find(FLobbyId LobbyId) const;
 	TFuture<TDefaultErrorResultInternal<TSharedRef<FLobbyDataEOS>>> FindOrCreateFromLobbyDetails(FAccountId LocalAccountId, const TSharedRef<FLobbyDetailsEOS>& LobbyDetails);
 
diff --git a/Engine/Plugins/Online/Experimental/OnlineServicesEOSGS/Source/Public/Online/LobbiesEOSGS.h b/Engine/Plugins/Online/Experimental/OnlineServicesEOSGS/Source/Public/Online/LobbiesEOSGS.h
--- a/Engine/Plugins/Online/Experimental/OnlineServicesEOSGS/Source/Public/Online/LobbiesEOSGS.h
+++ b/Engine/Plugins/Online/Experimental/OnlineServicesEOSGS/Source/Public/Online/LobbiesEOSGS.h
@@ -22,51 +22,6 @@ class FLobbySearchEOS;
 struct FClientLobbyDataChanges;
 struct FClientLobbyMemberDataChanges;
 
-struct FLobbiesJoinLobbyImpl
-{
-	static constexpr TCHAR Name[] = TEXT("JoinLobbyImpl");
-
-	struct Params
-	{
-		// The lobby handle data.
-		TSharedPtr<FLobbyDataEOS> LobbyData;
-
-		// The local user agent which will perform the action.
-		FAccountId LocalAccountId;
-
-		// The local name for the lobby.
-		FName LocalName;
-
-		// Local users who will be joining the lobby.
-		TArray<FJoinLobbyLocalUserData> LocalUsers;
-	};
-
-	struct Result
-	{
-	};
-};
-
-struct FLobbiesJoinLobbyMemberImpl
-{
-	static constexpr TCHAR Name[] = TEXT("JoinLobbyMemberImpl");
-
-	struct Params
-	{
-		// The lobby handle data.
-		TSharedPtr<FLobbyDataEOS> LobbyData;
-
-		// The local user agent which will perform the action.
-		FAccountId LocalAccountId;
-
-		// Initial attributes.
-		TMap<FLobbyAttributeId, FLobbyVariant> Attributes;
-	};
-
-	struct Result
-	{
-	};
-};
-
 struct FLobbiesLeaveLobbyImpl
 {
 	static constexpr TCHAR Name[] = TEXT("LeaveLobbyImpl");
@@ -91,8 +46,8 @@ struct FLobbiesDestroyLobbyImpl
 
 	struct Params
 	{
-		// The lobby handle data.
-		TSharedPtr<FLobbyDataEOS> LobbyData;
+		// The name of the lobby to be destroyed.
+		FString LobbyIdString;
 
 		// The local user agent which will perform the action.
 		FAccountId LocalAccountId;
@@ -302,9 +257,6 @@ protected:
 	TSharedPtr<FLobbyInviteDataEOS> GetActiveInvite(FAccountId TargetUser, FLobbyId TargetLobbyId);
 
 	// LobbyData will be fetched from the operation data if not set in Params.
-	TFuture<TDefaultErrorResult<FLobbiesJoinLobbyImpl>> JoinLobbyImpl(FLobbiesJoinLobbyImpl::Params&& Params);
-	TOnlineAsyncOpHandle<FLobbiesJoinLobbyMemberImpl> JoinLobbyMemberImplOp(FLobbiesJoinLobbyMemberImpl::Params&& Params);
-	TFuture<TDefaultErrorResult<FLobbiesJoinLobbyMemberImpl>> JoinLobbyMemberImpl(FLobbiesJoinLobbyMemberImpl::Params&& Params);
 	TFuture<TDefaultErrorResult<FLobbiesLeaveLobbyImpl>> LeaveLobbyImpl(FLobbiesLeaveLobbyImpl::Params&& Params);
 	TFuture<TDefaultErrorResult<FLobbiesDestroyLobbyImpl>> DestroyLobbyImpl(FLobbiesDestroyLobbyImpl::Params&& Params);
 	TFuture<TDefaultErrorResult<FLobbiesInviteLobbyMemberImpl>> InviteLobbyMemberImpl(FLobbiesInviteLobbyMemberImpl::Params&& Params);
@@ -332,25 +284,6 @@ protected:
 
 namespace Meta {
 
-BEGIN_ONLINE_STRUCT_META(FLobbiesJoinLobbyImpl::Params)
-	ONLINE_STRUCT_FIELD(FLobbiesJoinLobbyImpl::Params, LobbyData),
-	ONLINE_STRUCT_FIELD(FLobbiesJoinLobbyImpl::Params, LocalAccountId),
-	ONLINE_STRUCT_FIELD(FLobbiesJoinLobbyImpl::Params, LocalName),
-	ONLINE_STRUCT_FIELD(FLobbiesJoinLobbyImpl::Params, LocalUsers)
-END_ONLINE_STRUCT_META()
-
-BEGIN_ONLINE_STRUCT_META(FLobbiesJoinLobbyImpl::Result)
-END_ONLINE_STRUCT_META()
-
-BEGIN_ONLINE_STRUCT_META(FLobbiesJoinLobbyMemberImpl::Params)
-	ONLINE_STRUCT_FIELD(FLobbiesJoinLobbyMemberImpl::Params, LobbyData),
-	ONLINE_STRUCT_FIELD(FLobbiesJoinLobbyMemberImpl::Params, LocalAccountId),
-	ONLINE_STRUCT_FIELD(FLobbiesJoinLobbyMemberImpl::Params, Attributes)
-END_ONLINE_STRUCT_META()
-
-BEGIN_ONLINE_STRUCT_META(FLobbiesJoinLobbyMemberImpl::Result)
-END_ONLINE_STRUCT_META()
-
 BEGIN_ONLINE_STRUCT_META(FLobbiesLeaveLobbyImpl::Params)
 	ONLINE_STRUCT_FIELD(FLobbiesLeaveLobbyImpl::Params, LobbyData),
 	ONLINE_STRUCT_FIELD(FLobbiesLeaveLobbyImpl::Params, LocalAccountId)
@@ -360,7 +293,7 @@ BEGIN_ONLINE_STRUCT_META(FLobbiesLeaveLobbyImpl::Result)
 END_ONLINE_STRUCT_META()
 
 BEGIN_ONLINE_STRUCT_META(FLobbiesDestroyLobbyImpl::Params)
-	ONLINE_STRUCT_FIELD(FLobbiesDestroyLobbyImpl::Params, LobbyData),
+	ONLINE_STRUCT_FIELD(FLobbiesDestroyLobbyImpl::Params, LobbyIdString),
 	ONLINE_STRUCT_FIELD(FLobbiesDestroyLobbyImpl::Params, LocalAccountId)
 END_ONLINE_STRUCT_META()
 
diff --git a/Engine/Plugins/Online/Experimental/OnlineServicesNull/Source/Private/Online/LobbiesNull.cpp b/Engine/Plugins/Online/Experimental/OnlineServicesNull/Source/Private/Online/LobbiesNull.cpp
--- a/Engine/Plugins/Online/Experimental/OnlineServicesNull/Source/Private/Online/LobbiesNull.cpp
+++ b/Engine/Plugins/Online/Experimental/OnlineServicesNull/Source/Private/Online/LobbiesNull.cpp
@@ -154,19 +154,14 @@ TSharedRef<FLobbyNull> FLobbiesNull::CreateNamedLobby(const FCreateLobby::Params
 	Lobby->Data->Attributes = Params.Attributes;
 	Lobby->Data->MaxMembers = 100; // temp
 
-	// construct the members
-	for (const FJoinLobbyLocalUserData& Data : Params.LocalUsers)
-	{
-		TSharedRef<FLobbyMember> NewMember = MakeShared<FLobbyMember>();
-
-		NewMember->AccountId = Data.LocalAccountId;
-		NewMember->PlatformAccountId = Data.LocalAccountId;
-		NewMember->PlatformDisplayName = TEXT("TEMP");
-		NewMember->Attributes = Data.Attributes;
-
-		Lobby->Data->Members.Add(Data.LocalAccountId, NewMember);
+	// construct the member data.
+	TSharedRef<FLobbyMember> NewMember = MakeShared<FLobbyMember>();
+	NewMember->AccountId = Params.LocalAccountId;
+	NewMember->PlatformAccountId = Params.LocalAccountId;
+	NewMember->PlatformDisplayName = TEXT("TEMP");
+	NewMember->Attributes = Params.UserAttributes;
 
-	}
+	Lobby->Data->Members.Add(Params.LocalAccountId, NewMember);
 
 	AllLobbies.Add(Lobby->Data->LobbyId, Lobby);
 	return Lobby;
