Refactor of UpdateSession into UpdateSessionSettings and UpdateSessionMember
Addition of auxiliary methods and operators to simplify update data operations

diff --git a/Engine/Plugins/Online/Experimental/OnlineServices/Source/OnlineServicesCommon/Private/Online/SessionsCommon.cpp b/Engine/Plugins/Online/Experimental/OnlineServices/Source/OnlineServicesCommon/Private/Online/SessionsCommon.cpp
--- a/Engine/Plugins/Online/Experimental/OnlineServices/Source/OnlineServicesCommon/Private/Online/SessionsCommon.cpp
+++ b/Engine/Plugins/Online/Experimental/OnlineServices/Source/OnlineServicesCommon/Private/Online/SessionsCommon.cpp
@@ -54,6 +54,16 @@ namespace UE::Online {
 		return *this;
 	}
 
+	/** FCombinedSessionUpdate */
+	FCombinedSessionUpdate& FCombinedSessionUpdate::operator+=(FCombinedSessionUpdate&& UpdatedValues)
+	{
+		UpdatedSessionSettings += MoveTemp(UpdatedValues.UpdatedSessionSettings);
+
+		UpdatedSessionMembers.Append(MoveTemp(UpdatedValues.UpdatedSessionMembers));
+
+		return *this;
+	}
+
 	/** FSessionsCommon */
 	FSessionsCommon::FSessionsCommon(FOnlineServicesCommon& InServices)
 		: TOnlineComponent(TEXT("Sessions"), InServices)
@@ -73,7 +83,8 @@ namespace UE::Online {
 		RegisterCommand(&FSessionsCommon::GetSessionByName);
 		RegisterCommand(&FSessionsCommon::GetSessionById);
 		RegisterCommand(&FSessionsCommon::CreateSession);
-		RegisterCommand(&FSessionsCommon::UpdateSession);
+		RegisterCommand(&FSessionsCommon::UpdateSessionSettings);
+		RegisterCommand(&FSessionsCommon::UpdateSessionMember);
 		RegisterCommand(&FSessionsCommon::LeaveSession);
 		RegisterCommand(&FSessionsCommon::FindSessions);
 		RegisterCommand(&FSessionsCommon::StartMatchmaking);
@@ -176,11 +187,96 @@ namespace UE::Online {
 		return Operation->GetHandle();
 	}
 
-	TOnlineAsyncOpHandle<FUpdateSession> FSessionsCommon::UpdateSession(FUpdateSession::Params&& Params)
+	TOnlineAsyncOpHandle<FUpdateSessionSettings> FSessionsCommon::UpdateSessionSettings(FUpdateSessionSettings::Params&& Params)
 	{
-		TOnlineAsyncOpRef<FUpdateSession> Operation = GetOp<FUpdateSession>(MoveTemp(Params));
-		Operation->SetError(Errors::NotImplemented());
-		return Operation->GetHandle();
+		TOnlineAsyncOpRef<FUpdateSessionSettings> Op = GetOp<FUpdateSessionSettings>(MoveTemp(Params));
+		const FUpdateSessionSettings::Params& OpParams = Op->GetParams();
+
+		FOnlineError ParamsCheck = CheckUpdateSessionSettingsParams(OpParams);
+		if (ParamsCheck != Errors::Success())
+		{
+			Op->SetError(MoveTemp(ParamsCheck));
+			return Op->GetHandle();
+		}
+
+		Op->Then([this](TOnlineAsyncOp<FUpdateSessionSettings>& Op) mutable
+		{
+			const FUpdateSessionSettings::Params& OpParams = Op.GetParams();
+
+			FOnlineError StateCheck = CheckUpdateSessionSettingsState(OpParams);
+			if (StateCheck != Errors::Success())
+			{
+				Op.SetError(MoveTemp(StateCheck));
+				return;
+			}
+
+			UpdateSessionImpl({ OpParams.LocalAccountId, OpParams.SessionName, { OpParams.Mutations, {} } })
+			.OnComplete([this, WeakOp = Op.AsWeak()](const TOnlineResult<FUpdateSessionImpl>& Result)
+			{
+				if (TOnlineAsyncOpPtr<FUpdateSessionSettings> StrongOp = WeakOp.Pin())
+				{
+					if (Result.IsOk())
+					{
+						StrongOp->SetResult({ });
+					}
+					else
+					{
+						FOnlineError ErrorValue = Result.GetErrorValue();
+						StrongOp->SetError(MoveTemp(ErrorValue));
+					}
+				}
+			});			
+		})
+		.Enqueue(GetSerialQueue());
+
+		return Op->GetHandle();
+	}
+
+	TOnlineAsyncOpHandle<FUpdateSessionMember> FSessionsCommon::UpdateSessionMember(FUpdateSessionMember::Params&& Params)
+	{
+		// TODO: check if session members can update their own settings across all APIs or if they need the owner to do it
+
+		TOnlineAsyncOpRef<FUpdateSessionMember> Op = GetOp<FUpdateSessionMember>(MoveTemp(Params));
+		const FUpdateSessionMember::Params& OpParams = Op->GetParams();
+
+		FOnlineError ParamsCheck = CheckUpdateSessionMemberParams(OpParams);
+		if (ParamsCheck != Errors::Success())
+		{
+			Op->SetError(MoveTemp(ParamsCheck));
+			return Op->GetHandle();
+		}
+
+		Op->Then([this](TOnlineAsyncOp<FUpdateSessionMember>& Op) mutable
+		{
+			const FUpdateSessionMember::Params& OpParams = Op.GetParams();
+
+			FOnlineError StateCheck = CheckUpdateSessionMemberState(OpParams);
+			if (StateCheck != Errors::Success())
+			{
+				Op.SetError(MoveTemp(StateCheck));
+				return;
+			}
+
+			UpdateSessionImpl({ OpParams.LocalAccountId, OpParams.SessionName, { FSessionSettingsUpdate() , { { OpParams.LocalAccountId, OpParams.Mutations } } } })
+			.OnComplete([this, WeakOp = Op.AsWeak()](const TOnlineResult<FUpdateSessionImpl>& Result)
+			{
+				if (TOnlineAsyncOpPtr<FUpdateSessionMember> StrongOp = WeakOp.Pin())
+				{
+					if (Result.IsOk())
+					{
+						StrongOp->SetResult({ });
+					}
+					else
+					{
+						FOnlineError ErrorValue = Result.GetErrorValue();
+						StrongOp->SetError(MoveTemp(ErrorValue));
+					}
+				}
+			});
+		})
+		.Enqueue(GetSerialQueue());
+
+		return Op->GetHandle();
 	}
 
 	TOnlineAsyncOpHandle<FLeaveSession> FSessionsCommon::LeaveSession(FLeaveSession::Params&& Params)
@@ -476,6 +572,84 @@ namespace UE::Online {
 		AllSessionsById.Remove(SessionId);
 	}
 
+#define COPY_TOPTIONAL_VALUE_IF_SET(Value) \
+	if (UpdatedValues.Value.IsSet()) \
+	{ \
+		Result.Value = UpdatedValues.Value.GetValue(); \
+	} \
+
+	FSessionSettingsChanges FSessionsCommon::BuildSessionSettingsChanges(const TSharedRef<FSessionCommon>& Session, const FSessionSettingsUpdate& UpdatedValues) const
+	{
+		FSessionSettingsChanges Result;
+
+		COPY_TOPTIONAL_VALUE_IF_SET(SchemaName) // TODO: We may need some additional logic for schema changes
+		COPY_TOPTIONAL_VALUE_IF_SET(NumMaxConnections)
+		COPY_TOPTIONAL_VALUE_IF_SET(JoinPolicy)
+		COPY_TOPTIONAL_VALUE_IF_SET(bAllowNewMembers)
+
+		Result.RemovedCustomSettings.Append(UpdatedValues.RemovedCustomSettings);
+
+		for (const TPair<FSchemaAttributeId, FCustomSessionSetting>& Entry : UpdatedValues.UpdatedCustomSettings)
+		{
+			if (const FCustomSessionSetting* CustomSetting = Session->GetSessionSettings().CustomSettings.Find(Entry.Key))
+			{
+				FCustomSessionSettingUpdate SettingUpdate = { *CustomSetting, Entry.Value };
+
+				Result.ChangedCustomSettings.Emplace(Entry.Key, SettingUpdate);
+			}
+			else
+			{
+				Result.AddedCustomSettings.Add(Entry);
+			}
+		}
+
+		return Result;
+	}
+
+#undef COPY_TOPTIONAL_VALUE_IF_SET
+
+	FSessionUpdate FSessionsCommon::BuildSessionUpdate(const TSharedRef<FSessionCommon>& Session, const FCombinedSessionUpdate& UpdatedValues) const
+	{
+		FSessionUpdate Result;
+
+		Result.SessionSettingsChanges = BuildSessionSettingsChanges(Session, UpdatedValues.UpdatedSessionSettings);
+
+		for (const TPair<FAccountId, FSessionMemberUpdate>& SessionMemberEntry : UpdatedValues.UpdatedSessionMembers)
+		{
+			if (const FSessionMember* SessionMember = Session->GetSessionMembers().Find(SessionMemberEntry.Key))
+			{
+				FSessionMemberChanges SessionMemberChanges;
+
+				SessionMemberChanges.RemovedMemberSettings.Append(SessionMemberEntry.Value.RemovedMemberSettings);
+
+				for (const TPair <FSchemaAttributeId, FCustomSessionSetting>& SettingEntry : SessionMemberEntry.Value.UpdatedMemberSettings)
+				{
+					if (const FCustomSessionSetting* MemberSetting = SessionMember->MemberSettings.Find(SettingEntry.Key))
+					{
+						FCustomSessionSettingUpdate SettingUpdate = { *MemberSetting, SettingEntry.Value };
+
+						SessionMemberChanges.ChangedMemberSettings.Emplace(SettingEntry.Key, SettingUpdate);
+					}
+					else
+					{
+						SessionMemberChanges.AddedMemberSettings.Add(SettingEntry);
+					}
+				}
+
+				Result.SessionMembersChanges.Emplace(SessionMemberEntry.Key, SessionMemberChanges);
+			}
+		}
+
+		return Result;
+	}
+
+	TOnlineAsyncOpHandle<FUpdateSessionImpl> FSessionsCommon::UpdateSessionImpl(FUpdateSessionImpl::Params&& Params)
+	{
+		TOnlineAsyncOpRef<FUpdateSessionImpl> Operation = GetOp<FUpdateSessionImpl>(MoveTemp(Params));
+		Operation->SetError(Errors::NotImplemented());
+		return Operation->GetHandle();
+	}
+
 	TOnlineResult<FAddSessionMember> FSessionsCommon::AddSessionMemberImpl(const FAddSessionMember::Params& Params)
 	{
 		TOnlineResult<FGetMutableSessionByName> GetMutableSessionByNameResult = GetMutableSessionByName({ Params.SessionName });
@@ -580,28 +754,91 @@ namespace UE::Online {
 		return Errors::Success();
 	}
 
-	FOnlineError FSessionsCommon::CheckUpdateSessionState(const FUpdateSession::Params& Params)
+	FOnlineError FSessionsCommon::CheckUpdateSessionSettingsParams(const FUpdateSessionSettings::Params& Params)
 	{
-		TOnlineResult<FGetSessionByName> GetSessionByNameResult = GetSessionByName({ Params.SessionName });
-		if (GetSessionByNameResult.IsOk())
+		if (!Params.LocalAccountId.IsValid())
 		{
-			TSharedRef<const ISession> FoundSession = GetSessionByNameResult.GetOkValue().Session;
+			return Errors::InvalidUser();
+		}
 
-			if (!FoundSession->GetSessionInfo().bIsDedicatedServerSession)
+		if (Params.SessionName.IsNone())
+		{
+			return Errors::InvalidParams();
+		}
+
+		for (const TPair<FSchemaAttributeId, FCustomSessionSetting>& Entry : Params.Mutations.UpdatedCustomSettings)
+		{
+			if (Entry.Key.IsNone())
 			{
-				if (!Services.GetAuthInterface()->IsLoggedIn(Params.LocalAccountId))
-				{
-					UE_LOG(LogTemp, Warning, TEXT("[FSessionsNull::CheckUpdateSessionState] Could not update session with user [%s] not logged in"), *ToLogString(Params.LocalAccountId));
+				return Errors::InvalidParams();
+			}
+		}
 
-					return Errors::InvalidUser();
-				}
+		return Errors::Success();
+	}
+
+	FOnlineError FSessionsCommon::CheckUpdateSessionSettingsState(const FUpdateSessionSettings::Params& Params)
+	{
+		// User login check
+		IAuthPtr Auth = Services.GetAuthInterface();
+		if (!Auth->IsLoggedIn(Params.LocalAccountId))
+		{
+			UE_LOG(LogTemp, Warning, TEXT("[FSessionsCommon::CheckUpdateSessionSettingsState] User [%s] not logged in"), *ToLogString(Params.LocalAccountId));
+
+			return Errors::InvalidUser();
+		}
+
+		// Session name check
+		if (TOptional<FOnlineError> Result = CheckSessionExistsByName(Params.LocalAccountId, Params.SessionName))
+		{
+			UE_LOG(LogTemp, Warning, TEXT("[FSessionsCommon::CheckUpdateSessionSettingsState] Session with name [%s] not found."), *Params.SessionName.ToString());
+
+			return Result.GetValue();
+		}
+
+		return Errors::Success();
+	}
+
+	FOnlineError FSessionsCommon::CheckUpdateSessionMemberParams(const FUpdateSessionMember::Params& Params)
+	{
+		if (!Params.LocalAccountId.IsValid())
+		{
+			return Errors::InvalidUser();
+		}
+
+		if (Params.SessionName.IsNone())
+		{
+			return Errors::InvalidParams();
+		}
+
+		for (const TPair<FSchemaAttributeId, FCustomSessionSetting>& Entry : Params.Mutations.UpdatedMemberSettings)
+		{
+			if (Entry.Key.IsNone())
+			{
+				return Errors::InvalidParams();
 			}
 		}
-		else
+
+		return Errors::Success();
+	}
+
+	FOnlineError FSessionsCommon::CheckUpdateSessionMemberState(const FUpdateSessionMember::Params& Params)
+	{
+		// User login check
+		IAuthPtr Auth = Services.GetAuthInterface();
+		if (!Auth->IsLoggedIn(Params.LocalAccountId))
+		{
+			UE_LOG(LogTemp, Warning, TEXT("[FSessionsCommon::CheckUpdateSessionMemberState] User [%s] not logged in"), *ToLogString(Params.LocalAccountId));
+
+			return Errors::InvalidUser();
+		}
+
+		// Session name check
+		if (TOptional<FOnlineError> Result = CheckSessionExistsByName(Params.LocalAccountId, Params.SessionName))
 		{
-			UE_LOG(LogTemp, Warning, TEXT("[FSessionsCommon::CheckUpdateSessionState] Could not update session with name [%s]. Session not found"), *Params.SessionName.ToString());
+			UE_LOG(LogTemp, Warning, TEXT("[FSessionsCommon::CheckUpdateSessionMemberState] Session with name [%s] not found."), *Params.SessionName.ToString());
 
-			return GetSessionByNameResult.GetErrorValue();
+			return Result.GetValue();
 		}
 
 		return Errors::Success();
diff --git a/Engine/Plugins/Online/Experimental/OnlineServices/Source/OnlineServicesCommon/Public/Online/SessionsCommon.h b/Engine/Plugins/Online/Experimental/OnlineServices/Source/OnlineServicesCommon/Public/Online/SessionsCommon.h
--- a/Engine/Plugins/Online/Experimental/OnlineServices/Source/OnlineServicesCommon/Public/Online/SessionsCommon.h
+++ b/Engine/Plugins/Online/Experimental/OnlineServices/Source/OnlineServicesCommon/Public/Online/SessionsCommon.h
@@ -198,6 +198,39 @@ struct FGetMutableSessionById
 	};
 };
 
+struct FCombinedSessionUpdate
+{
+	/** Changes to current session settings */
+	FSessionSettingsUpdate UpdatedSessionSettings;
+
+	/** Map of changes to current member settings */
+	FSessionMemberUpdatesMap UpdatedSessionMembers;
+
+	FCombinedSessionUpdate& operator+=(FCombinedSessionUpdate&& UpdatedValue);
+};
+
+struct FUpdateSessionImpl
+{
+	static constexpr TCHAR Name[] = TEXT("UpdateSessionImpl");
+
+	struct Params
+	{
+		/** The local user agent which will perform the action. */
+		FAccountId LocalAccountId;
+
+		/** The local name for the session */
+		FName SessionName;
+
+		/** Combined data for a session update */
+		FCombinedSessionUpdate Mutations;
+	};
+
+	struct Result
+	{
+
+	};
+};
+
 class ONLINESERVICESCOMMON_API FSessionsCommon : public TOnlineComponent<ISessions>
 {
 public:
@@ -218,7 +251,8 @@ public:
 	virtual TOnlineResult<FSetPresenceSession> SetPresenceSession(FSetPresenceSession::Params&& Params) override;
 	virtual TOnlineResult<FClearPresenceSession> ClearPresenceSession(FClearPresenceSession::Params&& Params) override;
 	virtual TOnlineAsyncOpHandle<FCreateSession> CreateSession(FCreateSession::Params&& Params) override;
-	virtual TOnlineAsyncOpHandle<FUpdateSession> UpdateSession(FUpdateSession::Params&& Params) override;
+	virtual TOnlineAsyncOpHandle<FUpdateSessionSettings> UpdateSessionSettings(FUpdateSessionSettings::Params&& Params) override;
+	virtual TOnlineAsyncOpHandle<FUpdateSessionMember> UpdateSessionMember(FUpdateSessionMember::Params&& Params) override;
 	virtual TOnlineAsyncOpHandle<FLeaveSession> LeaveSession(FLeaveSession::Params&& Params) override;
 	virtual TOnlineAsyncOpHandle<FFindSessions> FindSessions(FFindSessions::Params&& Params) override;
 	virtual TOnlineAsyncOpHandle<FStartMatchmaking> StartMatchmaking(FStartMatchmaking::Params&& Params) override;
@@ -247,13 +281,21 @@ protected:
 	void ClearSessionInvitesForSession(const FAccountId& LocalAccountId, const FOnlineSessionId SessionId);
 	void ClearSessionReferences(const FOnlineSessionId SessionId, const FName& SessionName, const FAccountId& LocalAccountId);
 
+	FSessionSettingsChanges BuildSessionSettingsChanges(const TSharedRef<FSessionCommon>& Session, const FSessionSettingsUpdate& UpdatedValues) const;
+	FSessionUpdate BuildSessionUpdate(const TSharedRef<FSessionCommon>& Session, const FCombinedSessionUpdate& UpdatedValues) const;
+
+	virtual TOnlineAsyncOpHandle<FUpdateSessionImpl> UpdateSessionImpl(FUpdateSessionImpl::Params&& Params);
 	TOnlineResult<FAddSessionMember> AddSessionMemberImpl(const FAddSessionMember::Params& Params);
 	TOnlineResult<FRemoveSessionMember> RemoveSessionMemberImpl(const FRemoveSessionMember::Params& Params);
 
 	FOnlineError CheckCreateSessionParams(const FCreateSession::Params& Params);
 	FOnlineError CheckCreateSessionState(const FCreateSession::Params& Params);
 
-	FOnlineError CheckUpdateSessionState(const FUpdateSession::Params& Params);
+	FOnlineError CheckUpdateSessionSettingsParams(const FUpdateSessionSettings::Params& Params);
+	FOnlineError CheckUpdateSessionSettingsState(const FUpdateSessionSettings::Params& Params);
+
+	FOnlineError CheckUpdateSessionMemberParams(const FUpdateSessionMember::Params& Params);
+	FOnlineError CheckUpdateSessionMemberState(const FUpdateSessionMember::Params& Params);
 
 	FOnlineError CheckFindSessionsParams(const FFindSessions::Params& Params);
 	FOnlineError CheckFindSessionsState(const FFindSessions::Params& Params);
@@ -331,6 +373,20 @@ BEGIN_ONLINE_STRUCT_META(FGetMutableSessionById::Result)
 	ONLINE_STRUCT_FIELD(FGetMutableSessionById::Result, Session)
 END_ONLINE_STRUCT_META()
 
+BEGIN_ONLINE_STRUCT_META(FCombinedSessionUpdate)
+	ONLINE_STRUCT_FIELD(FCombinedSessionUpdate, UpdatedSessionSettings),
+	ONLINE_STRUCT_FIELD(FCombinedSessionUpdate, UpdatedSessionMembers)
+END_ONLINE_STRUCT_META()
+
+BEGIN_ONLINE_STRUCT_META(FUpdateSessionImpl::Params)
+	ONLINE_STRUCT_FIELD(FUpdateSessionImpl::Params, LocalAccountId),
+	ONLINE_STRUCT_FIELD(FUpdateSessionImpl::Params, SessionName),
+	ONLINE_STRUCT_FIELD(FUpdateSessionImpl::Params, Mutations)
+END_ONLINE_STRUCT_META()
+
+BEGIN_ONLINE_STRUCT_META(FUpdateSessionImpl::Result)
+END_ONLINE_STRUCT_META()
+
 /* Meta*/ }
 
 /* UE::Online */ }
diff --git a/Engine/Plugins/Online/Experimental/OnlineServices/Source/OnlineServicesCommonEngineUtils/Private/Online/SessionsLAN.cpp b/Engine/Plugins/Online/Experimental/OnlineServices/Source/OnlineServicesCommonEngineUtils/Private/Online/SessionsLAN.cpp
--- a/Engine/Plugins/Online/Experimental/OnlineServices/Source/OnlineServicesCommonEngineUtils/Private/Online/SessionsLAN.cpp
+++ b/Engine/Plugins/Online/Experimental/OnlineServices/Source/OnlineServicesCommonEngineUtils/Private/Online/SessionsLAN.cpp
@@ -169,20 +169,13 @@ TOnlineAsyncOpHandle<FCreateSession> FSessionsLAN::CreateSession(FCreateSession:
 	return Op->GetHandle();
 }
 
-TOnlineAsyncOpHandle<FUpdateSession> FSessionsLAN::UpdateSession(FUpdateSession::Params&& Params)
+TOnlineAsyncOpHandle<FUpdateSessionImpl> FSessionsLAN::UpdateSessionImpl(FUpdateSessionImpl::Params&& Params)
 {
-	TOnlineAsyncOpRef<FUpdateSession> Op = GetOp<FUpdateSession>(MoveTemp(Params));
+	TOnlineAsyncOpRef<FUpdateSessionImpl> Op = GetOp<FUpdateSessionImpl>(MoveTemp(Params));
 
-	Op->Then([this](TOnlineAsyncOp<FUpdateSession>& Op) mutable
+	Op->Then([this](TOnlineAsyncOp<FUpdateSessionImpl>& Op) mutable
 	{
-		const FUpdateSession::Params& OpParams = Op.GetParams();
-
-		FOnlineError StateCheck = CheckUpdateSessionState(OpParams);
-		if (StateCheck != Errors::Success())
-		{
-			Op.SetError(MoveTemp(StateCheck));
-			return;
-		}
+		const FUpdateSessionImpl::Params& OpParams = Op.GetParams();
 
 		// We'll check and update all settings one by one, with additional logic wherever required
 
@@ -191,13 +184,13 @@ TOnlineAsyncOpHandle<FUpdateSession> FSessionsLAN::UpdateSession(FUpdateSession:
 		
 		TSharedRef<FSessionCommon> FoundSession = GetMutableSessionByNameResult.GetOkValue().Session;
 		FSessionSettings& SessionSettings = FoundSession->SessionSettings;
-		const FSessionSettingsUpdate& UpdatedSettings = OpParams.Mutations;
+		const FCombinedSessionUpdate& SessionUpdate = OpParams.Mutations;
 
-		if (UpdatedSettings.JoinPolicy.IsSet())
+		if (SessionUpdate.UpdatedSessionSettings.JoinPolicy.IsSet())
 		{
 			// Changes in the join policy setting will mean we start or stop the LAN Beacon broadcasting the session information
 			// There is no FriendsLAN interface at the time of this implementation, hence the binary behavior (Public/Non-Public)
-			if (SessionSettings.JoinPolicy != ESessionJoinPolicy::Public && UpdatedSettings.JoinPolicy.GetValue() == ESessionJoinPolicy::Public)
+			if (SessionSettings.JoinPolicy != ESessionJoinPolicy::Public && SessionUpdate.UpdatedSessionSettings.JoinPolicy.GetValue() == ESessionJoinPolicy::Public)
 			{
 				if (TOptional<FOnlineError> TryHostLANSessionResult = TryHostLANSession())
 				{
@@ -205,21 +198,22 @@ TOnlineAsyncOpHandle<FUpdateSession> FSessionsLAN::UpdateSession(FUpdateSession:
 					return;
 				}
 			}
-			else if (SessionSettings.JoinPolicy == ESessionJoinPolicy::Public && UpdatedSettings.JoinPolicy.GetValue() != ESessionJoinPolicy::Public)
+			else if (SessionSettings.JoinPolicy == ESessionJoinPolicy::Public && SessionUpdate.UpdatedSessionSettings.JoinPolicy.GetValue() != ESessionJoinPolicy::Public)
 			{
 				StopLANSession();
 			}
 		}
 
-		//SessionSettings += UpdatedSettings;
+		// We update our local session
+		FSessionUpdate SessionUpdateData = BuildSessionUpdate(FoundSession, OpParams.Mutations);
+
+		(*FoundSession) += SessionUpdateData;
 
 		// We set the result and fire the event
 		Op.SetResult({ });
 
-		//FSessionUpdated SessionUpdatedEvent{ OpParams.SessionName, UpdatedSettings };
-		//SessionEvents.OnSessionUpdated.Broadcast(SessionUpdatedEvent);
-
-		// TODO: Will be refactored as part of the UpdateSettings split
+		FSessionUpdated SessionUpdatedEvent{ OpParams.SessionName, SessionUpdateData };
+		SessionEvents.OnSessionUpdated.Broadcast(SessionUpdatedEvent);
 	})
 	.Enqueue(GetSerialQueue());
 
diff --git a/Engine/Plugins/Online/Experimental/OnlineServices/Source/OnlineServicesCommonEngineUtils/Public/Online/SessionsLAN.h b/Engine/Plugins/Online/Experimental/OnlineServices/Source/OnlineServicesCommonEngineUtils/Public/Online/SessionsLAN.h
--- a/Engine/Plugins/Online/Experimental/OnlineServices/Source/OnlineServicesCommonEngineUtils/Public/Online/SessionsLAN.h
+++ b/Engine/Plugins/Online/Experimental/OnlineServices/Source/OnlineServicesCommonEngineUtils/Public/Online/SessionsLAN.h
@@ -57,7 +57,7 @@ public:
 
 	// ISessions
 	virtual TOnlineAsyncOpHandle<FCreateSession> CreateSession(FCreateSession::Params&& Params) override;
-	virtual TOnlineAsyncOpHandle<FUpdateSession> UpdateSession(FUpdateSession::Params&& Params) override;
+	virtual TOnlineAsyncOpHandle<FUpdateSessionImpl> UpdateSessionImpl(FUpdateSessionImpl::Params&& Params) override;
 	virtual TOnlineAsyncOpHandle<FFindSessions> FindSessions(FFindSessions::Params&& Params) override;
 	virtual TOnlineAsyncOpHandle<FJoinSession> JoinSession(FJoinSession::Params&& Params) override;
 	virtual TOnlineAsyncOpHandle<FLeaveSession> LeaveSession(FLeaveSession::Params&& Params) override;
diff --git a/Engine/Plugins/Online/Experimental/OnlineServices/Source/OnlineServicesInterface/Private/Online/Sessions.cpp b/Engine/Plugins/Online/Experimental/OnlineServices/Source/OnlineServicesInterface/Private/Online/Sessions.cpp
--- a/Engine/Plugins/Online/Experimental/OnlineServices/Source/OnlineServicesInterface/Private/Online/Sessions.cpp
+++ b/Engine/Plugins/Online/Experimental/OnlineServices/Source/OnlineServicesInterface/Private/Online/Sessions.cpp
@@ -36,49 +36,16 @@
 	}
 }
 
-#define COPY_TOPTIONAL_VALUE_IF_SET(Value) \
-	if (SettingsChanges.Value.IsSet()) \
-	{ \
-		Value = SettingsChanges.Value.GetValue(); \
-	} \
-
-FSessionSettings& FSessionSettings::operator+=(const FSessionSettingsChanges& SettingsChanges)
-{
-	COPY_TOPTIONAL_VALUE_IF_SET(SchemaName) // TODO: We may need some additional logic for schema changes
-	COPY_TOPTIONAL_VALUE_IF_SET(NumMaxConnections)
-	COPY_TOPTIONAL_VALUE_IF_SET(JoinPolicy)
-	COPY_TOPTIONAL_VALUE_IF_SET(bAllowNewMembers)
-		
-	for (const FName& Key : SettingsChanges.RemovedCustomSettings)
-	{
-		CustomSettings.Remove(Key);
-	}
-
-	CustomSettings.Append(SettingsChanges.AddedCustomSettings);
-
-	for (const TPair<FName, FCustomSessionSettingUpdate>& SettingEntry : SettingsChanges.ChangedCustomSettings)
-	{
-		if (FCustomSessionSetting* CustomSetting = CustomSettings.Find(SettingEntry.Key))
-		{
-			(*CustomSetting) = SettingEntry.Value.NewValue;
-		}
-	}
-
-	return *this;
-}
-
-#undef COPY_TOPTIONAL_VALUE_IF_SET
-
 FSessionMemberUpdate& FSessionMemberUpdate::operator+=(FSessionMemberUpdate&& UpdatedValue)
 {
-	for (const TPair<FName, FCustomSessionSetting>& UpdatedMemberSetting : UpdatedValue.UpdatedMemberSettings)
+	for (const TPair<FSchemaAttributeId, FCustomSessionSetting>& UpdatedMemberSetting : UpdatedValue.UpdatedMemberSettings)
 	{
 		// If an update adds a modification to a setting that had previously been marked for removal, we'll keep the latest change
 		RemovedMemberSettings.Remove(UpdatedMemberSetting.Key);
 	}
 	UpdatedMemberSettings.Append(MoveTemp(UpdatedValue.UpdatedMemberSettings));
 
-	for (FName& Key : UpdatedValue.RemovedMemberSettings)
+	for (FSchemaAttributeId& Key : UpdatedValue.RemovedMemberSettings)
 	{
 		// If an update removes a setting that had previously been modified, we'll keep the latest change
 		UpdatedMemberSettings.Remove(Key);
@@ -86,3 +53,23 @@
 	}
 
 	return *this;
+}
+
+FSessionMember& FSessionMember::operator+=(const FSessionMemberChanges& MemberChanges)
+{
+	for (const FName& Key : MemberChanges.RemovedMemberSettings)
+	{
+		MemberSettings.Remove(Key);
+	}
+
+	MemberSettings.Append(MemberChanges.AddedMemberSettings);
+
+	for (const TPair<FName, FCustomSessionSettingUpdate>& SettingEntry : MemberChanges.ChangedMemberSettings)
+	{
+		if (FCustomSessionSetting* CustomSetting = MemberSettings.Find(SettingEntry.Key))
+		{
+			(*CustomSetting) = SettingEntry.Value.NewValue;
+		}
+	}
+
+	return *this;
@@ -101,40 +88,58 @@
 	MOVE_TOPTIONAL_IF_SET(JoinPolicy)
 	MOVE_TOPTIONAL_IF_SET(bAllowNewMembers)
 
-	for (const TPair<FName, FCustomSessionSetting>& UpdatedCustomSetting : UpdatedValue.UpdatedCustomSettings)
+	for (const TPair<FSchemaAttributeId, FCustomSessionSetting>& UpdatedCustomSetting : UpdatedValue.UpdatedCustomSettings)
 	{
 		// If an update adds a modification to a setting that had previously been marked for removal, we'll keep the latest change
 		RemovedCustomSettings.Remove(UpdatedCustomSetting.Key);
 	}
+
 	UpdatedCustomSettings.Append(MoveTemp(UpdatedValue.UpdatedCustomSettings));
 
-	for (FName& Key : UpdatedValue.RemovedCustomSettings)
+	for (FSchemaAttributeId& Key : UpdatedValue.RemovedCustomSettings)
 	{
 		// If an update removes a setting that had previously been modified, we'll keep the latest change
 		UpdatedCustomSettings.Remove(Key);
 		RemovedCustomSettings.AddUnique(MoveTemp(Key));
 	}
 
-	for (TPair<FAccountId, FSessionMemberUpdate>& UpdatedSessionMember : UpdatedValue.UpdatedSessionMembers)
-	{
-		// If an update adds a modification to a member that had previously been marked for removal, we'll keep the latest change
-		RemovedSessionMembers.Remove(UpdatedSessionMember.Key);
-
-		FSessionMemberUpdate MemberUpdate = UpdatedSessionMembers.FindOrAdd(MoveTemp(UpdatedSessionMember.Key));
-		MemberUpdate += MoveTemp(UpdatedSessionMember.Value);
-	}
-
-	for (FAccountId& Key : UpdatedValue.RemovedSessionMembers)
-	{
-		// If an update removes a member that had previously been modified, we'll keep the latest change
-		UpdatedSessionMembers.Remove(Key);
-		RemovedSessionMembers.AddUnique(MoveTemp(Key));
-	}
-
 	return *this;
 }
 
 #undef MOVE_TOPTIONAL_IF_SET
 
+#define COPY_TOPTIONAL_VALUE_IF_SET(Value) \
+	if (SettingsChanges.Value.IsSet()) \
+	{ \
+		Value = SettingsChanges.Value.GetValue(); \
+	} \
+
+FSessionSettings& FSessionSettings::operator+=(const FSessionSettingsChanges& SettingsChanges)
+{
+	COPY_TOPTIONAL_VALUE_IF_SET(SchemaName) // TODO: We may need some additional logic for schema changes
+	COPY_TOPTIONAL_VALUE_IF_SET(NumMaxConnections)
+	COPY_TOPTIONAL_VALUE_IF_SET(JoinPolicy)
+	COPY_TOPTIONAL_VALUE_IF_SET(bAllowNewMembers)
+
+	for (const FSchemaAttributeId& Key : SettingsChanges.RemovedCustomSettings)
+	{
+		CustomSettings.Remove(Key);
+	}
+
+	CustomSettings.Append(SettingsChanges.AddedCustomSettings);
+
+	for (const TPair<FSchemaAttributeId, FCustomSessionSettingUpdate>& SettingEntry : SettingsChanges.ChangedCustomSettings)
+	{
+		if (FCustomSessionSetting* CustomSetting = CustomSettings.Find(SettingEntry.Key))
+		{
+			(*CustomSetting) = SettingEntry.Value.NewValue;
+		}
+	}
+
+	return *this;
+}
+
+#undef COPY_TOPTIONAL_VALUE_IF_SET
+
 const FString ToLogString(const ISession& Session)
 {
diff --git a/Engine/Plugins/Online/Experimental/OnlineServices/Source/OnlineServicesInterface/Public/Online/Sessions.h b/Engine/Plugins/Online/Experimental/OnlineServices/Source/OnlineServicesInterface/Public/Online/Sessions.h
--- a/Engine/Plugins/Online/Experimental/OnlineServices/Source/OnlineServicesInterface/Public/Online/Sessions.h
+++ b/Engine/Plugins/Online/Experimental/OnlineServices/Source/OnlineServicesInterface/Public/Online/Sessions.h
@@ -22,6 +22,8 @@ struct FFindSessionsSearchFilter
 	FSchemaVariant Value;
 };
 
+// Custom Session Settings
+
 struct FCustomSessionSetting
 {
 	/** Setting value */
@@ -45,15 +47,10 @@ struct FCustomSessionSettingUpdate
 
 using FCustomSessionSettingUpdateMap = TMap<FName, FCustomSessionSettingUpdate>;
 
-/** A member is a player that is part of the session, and it stops being a member when they leave it */
-struct FSessionMember
-{
-	FCustomSessionSettingsMap MemberSettings;
-};
+// Session Member Settings
 
-using FSessionMembersMap = TMap<FAccountId, FSessionMember>;
-
-struct FSessionMemberUpdate
+/** Contains new values for a SessionMember's attributes. Taken as a parameter by the FUpdateSessionMember method */
+struct ONLINESERVICESINTERFACE_API FSessionMemberUpdate
 {
 	FCustomSessionSettingsMap UpdatedMemberSettings;
 	TArray<FSchemaAttributeId> RemovedMemberSettings;
@@ -63,6 +60,32 @@ struct FSessionMemberUpdate
 
 using FSessionMemberUpdatesMap = TMap<FAccountId, FSessionMemberUpdate>;
 
+/** Contains updated data for any modifiable attributes of FSessionMember. Member of FSessionUpdated event */
+struct FSessionMemberChanges
+{
+	/** New custom settings, with their values */
+	FCustomSessionSettingsMap AddedMemberSettings;
+
+	/** Existing custom settings that changed value, including new and old values */
+	FCustomSessionSettingUpdateMap ChangedMemberSettings;
+
+	/** Keys for removed custom settings */
+	TArray<FName> RemovedMemberSettings;
+};
+
+using FSessionMemberChangesMap = TMap<FAccountId, FSessionMemberChanges>;
+
+/** A member is a player that is part of the session, and they stop being a member when they leave it */
+struct ONLINESERVICESINTERFACE_API FSessionMember
+{
+	/* Map of user-defined settings to be passed to the platform APIs as additional information for the session member */
+	FCustomSessionSettingsMap MemberSettings;
+
+	FSessionMember& operator+=(const FSessionMemberChanges& Changes);
+};
+
+using FSessionMembersMap = TMap<FAccountId, FSessionMember>;
+
 enum class ESessionJoinPolicy : uint8
 {
 	Public,
@@ -73,7 +96,7 @@ ONLINESERVICESINTERFACE_API const TCHAR* LexToString(ESessionJoinPolicy Value);
 ONLINESERVICESINTERFACE_API void LexFromString(ESessionJoinPolicy& Value, const TCHAR* InStr);
 
 /** Contains new values for an FSessions modifiable settings. Taken as a parameter by FUpdateSessions method */
-struct FSessionSettingsUpdate
+struct ONLINESERVICESINTERFACE_API FSessionSettingsUpdate
 {
 	/** Set with an updated value if the SchemaName field will be changed in the update operation */
 	TOptional<FSchemaId> SchemaName;
@@ -89,11 +112,6 @@ struct FSessionSettingsUpdate
 	/** Names of custom settings to be removed in the update operation*/
 	TArray<FSchemaAttributeId> RemovedCustomSettings;
 
-	/** Updated values for session member info to change in the update operation*/
-	FSessionMemberUpdatesMap UpdatedSessionMembers;
-	/** Id handles for session members to be removed in the update operation*/
-	TArray<FAccountId> RemovedSessionMembers;
-
 	FSessionSettingsUpdate& operator+=(FSessionSettingsUpdate&& UpdatedValue);
 };
 
@@ -119,19 +137,6 @@ struct FSessionSettingsChanges
 	TArray<FName> RemovedCustomSettings;
 };
 
-/** Contains updated data for any modifiable members of FSessionMember. Member of FSessionUpdated event */
-struct FSessionMemberChanges
-{
-	/** New custom settings, with their values */
-	FCustomSessionSettingsMap AddedMemberSettings;
-
-	/** Existing custom settings that changed value, including new and old values */
-	FCustomSessionSettingUpdateMap ChangedMemberSettings;
-
-	/** Keys for removed custom settings */
-	TArray<FName> RemovedMemberSettings;
-};
-
 /** Set of all of an FSession's defining properties that can be updated by the session owner during its lifetime */
 struct ONLINESERVICESINTERFACE_API FSessionSettings
 {
@@ -365,9 +370,9 @@ struct FCreateSession
 	};
 };
 
-struct FUpdateSession
+struct FUpdateSessionSettings
 {
-	static constexpr TCHAR Name[] = TEXT("UpdateSession");
+	static constexpr TCHAR Name[] = TEXT("UpdateSessionSettings");
 
 	struct Params
 	{
@@ -387,6 +392,28 @@ struct FUpdateSession
 	};
 };
 
+struct FUpdateSessionMember
+{
+	static constexpr TCHAR Name[] = TEXT("UpdateSessionMember");
+
+	struct Params
+	{
+		/** The local user agent which will perform the action. */
+		FAccountId LocalAccountId;
+
+		/** The local name for the session */
+		FName SessionName;
+
+		/** Changes to current session settings */
+		FSessionMemberUpdate Mutations;
+	};
+
+	struct Result
+	{
+
+	};
+};
+
 struct FLeaveSession
 {
 	static constexpr TCHAR Name[] = TEXT("LeaveSession");
@@ -602,8 +629,6 @@ struct FSessionLeft
 	FAccountId LocalAccountId;
 };
 
-using FSessionMemberChangesMap = TMap<FAccountId, FSessionMemberChanges>;
-
 /** Contains updated data for any modifiable members of ISession */
 struct ONLINESERVICESINTERFACE_API FSessionUpdate
 {
@@ -738,12 +763,20 @@ public:
 	virtual TOnlineAsyncOpHandle<FCreateSession> CreateSession(FCreateSession::Params&& Params) = 0;
 
 	/**
-	 * Update a given session's settings.
+	 * Update a given session's settings. Can only be called by the session owner.
 	 *
-	 * @param Parameters for the UpdateSession call
+	 * @param Parameters for the UpdateSessionSettings call
 	 * @return
 	 */
-	virtual TOnlineAsyncOpHandle<FUpdateSession> UpdateSession(FUpdateSession::Params&& Params) = 0;
+	virtual TOnlineAsyncOpHandle<FUpdateSessionSettings> UpdateSessionSettings(FUpdateSessionSettings::Params&& Params) = 0;
+
+	/**
+	 * Update a given session member's information. Can be called by any session member to change only their own information.
+	 *
+	 * @param Parameters for the UpdateSessionMember call
+	 * @return
+	 */
+	virtual TOnlineAsyncOpHandle<FUpdateSessionMember> UpdateSessionMember(FUpdateSessionMember::Params&& Params) = 0;
 
 	/**
 	 * Leave and optionally destroy a given session.
@@ -915,9 +948,7 @@ BEGIN_ONLINE_STRUCT_META(FSessionSettingsUpdate)
 	ONLINE_STRUCT_FIELD(FSessionSettingsUpdate, JoinPolicy),
 	ONLINE_STRUCT_FIELD(FSessionSettingsUpdate, bAllowNewMembers),
 	ONLINE_STRUCT_FIELD(FSessionSettingsUpdate, UpdatedCustomSettings),
-	ONLINE_STRUCT_FIELD(FSessionSettingsUpdate, RemovedCustomSettings),
-	ONLINE_STRUCT_FIELD(FSessionSettingsUpdate, UpdatedSessionMembers),
-	ONLINE_STRUCT_FIELD(FSessionSettingsUpdate, RemovedSessionMembers)
+	ONLINE_STRUCT_FIELD(FSessionSettingsUpdate, RemovedCustomSettings)
 END_ONLINE_STRUCT_META()
 
 BEGIN_ONLINE_STRUCT_META(FSessionInvite)
@@ -1000,13 +1031,22 @@ END_ONLINE_STRUCT_META()
 BEGIN_ONLINE_STRUCT_META(FCreateSession::Result)
 END_ONLINE_STRUCT_META()
 
-BEGIN_ONLINE_STRUCT_META(FUpdateSession::Params)
-	ONLINE_STRUCT_FIELD(FUpdateSession::Params, LocalAccountId),
-	ONLINE_STRUCT_FIELD(FUpdateSession::Params, SessionName),
-	ONLINE_STRUCT_FIELD(FUpdateSession::Params, Mutations)
+BEGIN_ONLINE_STRUCT_META(FUpdateSessionSettings::Params)
+	ONLINE_STRUCT_FIELD(FUpdateSessionSettings::Params, LocalAccountId),
+	ONLINE_STRUCT_FIELD(FUpdateSessionSettings::Params, SessionName),
+	ONLINE_STRUCT_FIELD(FUpdateSessionSettings::Params, Mutations)
+END_ONLINE_STRUCT_META()
+
+BEGIN_ONLINE_STRUCT_META(FUpdateSessionSettings::Result)
+END_ONLINE_STRUCT_META()
+
+BEGIN_ONLINE_STRUCT_META(FUpdateSessionMember::Params)
+	ONLINE_STRUCT_FIELD(FUpdateSessionMember::Params, LocalAccountId),
+	ONLINE_STRUCT_FIELD(FUpdateSessionMember::Params, SessionName),
+	ONLINE_STRUCT_FIELD(FUpdateSessionMember::Params, Mutations)
 END_ONLINE_STRUCT_META()
 
-BEGIN_ONLINE_STRUCT_META(FUpdateSession::Result)
+BEGIN_ONLINE_STRUCT_META(FUpdateSessionMember::Result)
 END_ONLINE_STRUCT_META()
 
 BEGIN_ONLINE_STRUCT_META(FLeaveSession::Params)
diff --git a/Engine/Plugins/Online/Experimental/OnlineServicesEOSGS/Source/Private/Online/SessionsEOSGS.cpp b/Engine/Plugins/Online/Experimental/OnlineServicesEOSGS/Source/Private/Online/SessionsEOSGS.cpp
--- a/Engine/Plugins/Online/Experimental/OnlineServicesEOSGS/Source/Private/Online/SessionsEOSGS.cpp
+++ b/Engine/Plugins/Online/Experimental/OnlineServicesEOSGS/Source/Private/Online/SessionsEOSGS.cpp
@@ -365,7 +365,7 @@ TOnlineAsyncOpHandle<FCreateSession> FSessionsEOSGS::CreateSession(FCreateSessio
 		{
 			Op.SetError(MoveTemp(StateCheck));
 
-			return MakeFulfilledPromise<TDefaultErrorResult<FUpdateSessionImpl>>().GetFuture();
+			return MakeFulfilledPromise<TDefaultErrorResult<FUpdateSessionImplEOSGS>>().GetFuture();
 		}
 
 		if (!OpParams.SessionIdOverride.IsEmpty())
@@ -378,7 +378,7 @@ TOnlineAsyncOpHandle<FCreateSession> FSessionsEOSGS::CreateSession(FCreateSessio
 
 				Op.SetError(Errors::InvalidParams());
 
-				return MakeFulfilledPromise<TDefaultErrorResult<FUpdateSessionImpl>>().GetFuture();
+				return MakeFulfilledPromise<TDefaultErrorResult<FUpdateSessionImplEOSGS>>().GetFuture();
 			}
 		}
 
@@ -389,7 +389,7 @@ TOnlineAsyncOpHandle<FCreateSession> FSessionsEOSGS::CreateSession(FCreateSessio
 
 			Op.SetError(Errors::InvalidParams());
 
-			return MakeFulfilledPromise<TDefaultErrorResult<FUpdateSessionImpl>>().GetFuture();
+			return MakeFulfilledPromise<TDefaultErrorResult<FUpdateSessionImplEOSGS>>().GetFuture();
 		}
 
 		// After all initial checks, we start the session creation operations
@@ -429,7 +429,7 @@ TOnlineAsyncOpHandle<FCreateSession> FSessionsEOSGS::CreateSession(FCreateSessio
 
 			Op.SetError(Errors::FromEOSResult(ResultCode));
 
-			return MakeFulfilledPromise<TDefaultErrorResult<FUpdateSessionImpl>>().GetFuture();
+			return MakeFulfilledPromise<TDefaultErrorResult<FUpdateSessionImplEOSGS>>().GetFuture();
 		}
 
 		// TODO: We could call EOS_SessionModification_SetHostAddress at this point, although it's not necessary
@@ -438,11 +438,11 @@ TOnlineAsyncOpHandle<FCreateSession> FSessionsEOSGS::CreateSession(FCreateSessio
 		WriteCreateSessionModificationHandle(SessionModificationHandle, OpParams);
 
 		// Always update joinability on session creation
-		FUpdateSessionImpl::Params UpdateSessionImplParams{ MakeShared<FSessionModificationHandleEOSGS>(SessionModificationHandle), FUpdateSessionJoinabilityParams{ OpParams.SessionName, OpParams.SessionSettings.bAllowNewMembers } };
+		FUpdateSessionImplEOSGS::Params UpdateSessionImplParams{ MakeShared<FSessionModificationHandleEOSGS>(SessionModificationHandle), FUpdateSessionJoinabilityParams{ OpParams.SessionName, OpParams.SessionSettings.bAllowNewMembers } };
 		
-		return UpdateSessionImpl(MoveTemp(UpdateSessionImplParams));
+		return UpdateSessionImplEOSGS(MoveTemp(UpdateSessionImplParams));
 	})
-	.Then([this](TOnlineAsyncOp<FCreateSession>& Op, TDefaultErrorResult<FUpdateSessionImpl>&& Result)
+	.Then([this](TOnlineAsyncOp<FCreateSession>& Op, TDefaultErrorResult<FUpdateSessionImplEOSGS>&& Result)
 	{
 		const FCreateSession::Params& OpParams = Op.GetParams();
 
@@ -622,7 +622,7 @@ void FSessionsEOSGS::WriteCreateSessionModificationHandle(EOS_HSessionModificati
 	}		
 }
 
-void FSessionsEOSGS::WriteUpdateSessionModificationHandle(EOS_HSessionModification& SessionModificationHandle, const FName& SessionName, const FSessionSettingsUpdate& NewSettings)
+void FSessionsEOSGS::WriteUpdateSessionModificationHandle(EOS_HSessionModification& SessionModificationHandle, const FSessionSettingsUpdate& NewSettings)
 {
 	// TODO: We have the option to call EOS_SessionModification_SetHostAddress in EOS, useful if the session owner changes
 
@@ -640,12 +640,6 @@ void FSessionsEOSGS::WriteUpdateSessionModificationHandle(EOS_HSessionModificati
 		AddAttribute(SessionModificationHandle, EOSGS_SCHEMA_NAME, { FSchemaVariant(NewSettings.SchemaName.GetValue().ToString()), ESchemaAttributeVisibility::Public });
 	}
 
-	TOnlineResult<FGetSessionByName> GetSessionByNameResult = GetSessionByName({ SessionName });
-	check(GetSessionByNameResult.IsOk());
-
-	TSharedRef<const ISession> FoundSession = GetSessionByNameResult.GetOkValue().Session;
-	const FSessionSettings& SessionSettings = FoundSession->GetSessionSettings();
-
 	if (NewSettings.JoinPolicy.IsSet())
 	{
 		SetPermissionLevel(SessionModificationHandle, NewSettings.JoinPolicy.GetValue());
@@ -674,67 +668,40 @@ void FSessionsEOSGS::WriteUpdateSessionModificationHandle(EOS_HSessionModificati
 	{
 		AddAttribute(SessionModificationHandle, Entry.Key, Entry.Value);
 	}
+}
 
-	// Session Members
-	
-	for (const FAccountId& SessionMemberId : NewSettings.RemovedSessionMembers)
-	{
-		// To remove a session member, we'll remove all their attributes
- 		if (const FSessionMember* SessionMember = FoundSession->GetSessionMembers().Find(SessionMemberId))
- 		{
- 			TArray<FName> SettingsKeyArray;
- 			SessionMember->MemberSettings.GenerateKeyArray(SettingsKeyArray);
- 			for (const FSchemaAttributeId& CustomSettingEntryKey : SettingsKeyArray)
- 			{
- 				const FSchemaAttributeId& Key = FName(LexToString(GetProductUserIdChecked(SessionMemberId)) + TEXT(":") + CustomSettingEntryKey.ToString());
-				RemoveAttribute(SessionModificationHandle, Key);
- 			}
- 		}
+void FSessionsEOSGS::WriteUpdateSessionModificationHandle(EOS_HSessionModification& SessionModificationHandle, const FAccountId& LocalAccountId, const FSessionMemberUpdate& MemberUpdate)
+{
+	for (const FSchemaAttributeId& CustomSettingEntryKey : MemberUpdate.RemovedMemberSettings)
+	{
+		const FSchemaAttributeId& Key = FSchemaAttributeId(LexToString(GetProductUserIdChecked(LocalAccountId)) + TEXT(":") + CustomSettingEntryKey.ToString());
+		RemoveAttribute(SessionModificationHandle, Key);
 	}
 
-	for (const TPair<FAccountId, FSessionMemberUpdate>& SessionMemberEntry : NewSettings.UpdatedSessionMembers)
+	for (const TPair<FSchemaAttributeId, FCustomSessionSetting>& CustomSettingEntry : MemberUpdate.UpdatedMemberSettings)
 	{
-		const FSessionMemberUpdate& SessionMemberUpdate = SessionMemberEntry.Value;
-
-		for (const FSchemaAttributeId& CustomSettingEntryKey : SessionMemberUpdate.RemovedMemberSettings)
-		{
-			const FSchemaAttributeId& Key = FSchemaAttributeId(LexToString(GetProductUserIdChecked(SessionMemberEntry.Key)) + TEXT(":") + CustomSettingEntryKey.ToString());
-			RemoveAttribute(SessionModificationHandle, Key);
-		}
-
-		for (const TPair<FSchemaAttributeId, FCustomSessionSetting>& CustomSettingEntry : SessionMemberUpdate.UpdatedMemberSettings)
-		{
-			const FSchemaAttributeId& Key = FSchemaAttributeId(LexToString(GetProductUserIdChecked(SessionMemberEntry.Key)) + TEXT(":") + CustomSettingEntry.Key.ToString());
-			AddAttribute(SessionModificationHandle, Key, CustomSettingEntry.Value);
-		}
+		const FSchemaAttributeId& Key = FSchemaAttributeId(LexToString(GetProductUserIdChecked(LocalAccountId)) + TEXT(":") + CustomSettingEntry.Key.ToString());
+		AddAttribute(SessionModificationHandle, Key, CustomSettingEntry.Value);
 	}
 }
 
-TOnlineAsyncOpHandle<FUpdateSession> FSessionsEOSGS::UpdateSession(FUpdateSession::Params&& Params)
+TOnlineAsyncOpHandle<FUpdateSessionImpl> FSessionsEOSGS::UpdateSessionImpl(FUpdateSessionImpl::Params&& Params)
 {
 	// LAN Sessions
 	TOnlineResult<FGetSessionByName> Result = GetSessionByName({ Params.SessionName });
 	if (Result.IsOk() && Result.GetOkValue().Session->GetSessionInfo().bIsLANSession)
 	{
-		return FSessionsLAN::UpdateSession(MoveTemp(Params));
+		return FSessionsLAN::UpdateSessionImpl(MoveTemp(Params));
 	}
 
 	// EOSGS Sessions
 
-	TOnlineAsyncOpRef<FUpdateSession> Op = GetOp<FUpdateSession>(MoveTemp(Params));
-	const FUpdateSession::Params& OpParams = Op->GetParams();
+	TOnlineAsyncOpRef<FUpdateSessionImpl> Op = GetOp<FUpdateSessionImpl>(MoveTemp(Params));
+	const FUpdateSessionImpl::Params& OpParams = Op->GetParams();
 
-	Op->Then([this](TOnlineAsyncOp<FUpdateSession>& Op)
+	Op->Then([this](TOnlineAsyncOp<FUpdateSessionImpl>& Op)
 	{
-		const FUpdateSession::Params& OpParams = Op.GetParams();
-
-		FOnlineError StateCheck = CheckUpdateSessionState(OpParams);
-		if (StateCheck != Errors::Success())
-		{
-			Op.SetError(MoveTemp(StateCheck));
-
-			return MakeFulfilledPromise<TDefaultErrorResult<FUpdateSessionImpl>>().GetFuture();
-		}
+		const FUpdateSessionImpl::Params& OpParams = Op.GetParams();
 
 		EOS_Sessions_UpdateSessionModificationOptions UpdateSessionModificationOptions = {};
 		UpdateSessionModificationOptions.ApiVersion = EOS_SESSIONS_UPDATESESSIONMODIFICATION_API_LATEST;
@@ -751,28 +718,33 @@ TOnlineAsyncOpHandle<FUpdateSession> FSessionsEOSGS::UpdateSession(FUpdateSessio
 
 			Op.SetError(Errors::FromEOSResult(ResultCode));
 
-			return MakeFulfilledPromise<TDefaultErrorResult<FUpdateSessionImpl>>().GetFuture();
+			return MakeFulfilledPromise<TDefaultErrorResult<FUpdateSessionImplEOSGS>>().GetFuture();
 		}
 
 		// After creating the session modification handle, we'll add all the updated data to it
-		WriteUpdateSessionModificationHandle(SessionModificationHandle, OpParams.SessionName, OpParams.Mutations);
+		WriteUpdateSessionModificationHandle(SessionModificationHandle, OpParams.Mutations.UpdatedSessionSettings);
+
+		for (const TPair<FAccountId, FSessionMemberUpdate>& Entry : OpParams.Mutations.UpdatedSessionMembers)
+		{
+			WriteUpdateSessionModificationHandle(SessionModificationHandle, Entry.Key, Entry.Value);
+		}
 
 		// Whether we update joinability or not will depend on inf bAllowNewMembers was set to a new value
-		FUpdateSessionImpl::Params UpdateSessionImplParams = 
+		FUpdateSessionImplEOSGS::Params UpdateSessionImplEOSGSParams =
 		{
-		MakeShared<FSessionModificationHandleEOSGS>(SessionModificationHandle),
-		OpParams.Mutations.bAllowNewMembers.IsSet() ?
-			FUpdateSessionJoinabilityParams{ OpParams.SessionName, OpParams.Mutations.bAllowNewMembers.GetValue() } :
-			TOptional<FUpdateSessionJoinabilityParams>() // If bAllowNewMembers is not set, the TOptional will be unset too
+			MakeShared<FSessionModificationHandleEOSGS>(SessionModificationHandle),
+			OpParams.Mutations.UpdatedSessionSettings.bAllowNewMembers.IsSet() ?
+				FUpdateSessionJoinabilityParams{ OpParams.SessionName, OpParams.Mutations.UpdatedSessionSettings.bAllowNewMembers.GetValue() } :
+				TOptional<FUpdateSessionJoinabilityParams>() // If bAllowNewMembers is not set, the TOptional will be unset too
 		};
 
-		return UpdateSessionImpl(MoveTemp(UpdateSessionImplParams));
+		return UpdateSessionImplEOSGS(MoveTemp(UpdateSessionImplEOSGSParams));
 	})
-	.Then([this](TOnlineAsyncOp<FUpdateSession>& Op, TDefaultErrorResult<FUpdateSessionImpl>&& Result)
+	.Then([this](TOnlineAsyncOp<FUpdateSessionImpl>& Op, TDefaultErrorResult<FUpdateSessionImplEOSGS>&& Result)
 	{
 		if (Result.IsOk())
 		{
-			const FUpdateSession::Params& OpParams = Op.GetParams();
+			const FUpdateSessionImpl::Params& OpParams = Op.GetParams();
 
 			TOnlineResult<FGetMutableSessionByName> GetMutableSessionByNameResult = GetMutableSessionByName({ OpParams.SessionName });
 			if (GetMutableSessionByNameResult.IsOk())
@@ -780,15 +752,15 @@ TOnlineAsyncOpHandle<FUpdateSession> FSessionsEOSGS::UpdateSession(FUpdateSessio
 				TSharedRef<FSessionCommon> FoundSession = GetMutableSessionByNameResult.GetOkValue().Session;
 
 				// Now that the API Session update has processed successfully, we'll update our local session with the same data
-				//FoundSession->SessionSettings += OpParams.Mutations;
+				FSessionUpdate SessionUpdateData = BuildSessionUpdate(FoundSession, OpParams.Mutations);
+
+				(*FoundSession) += SessionUpdateData;
 
 				// We set the result and fire the event
 				Op.SetResult({ });
 
-				//FSessionUpdated SessionUpdatedEvent{ OpParams.SessionName, OpParams.Mutations };
-				//SessionEvents.OnSessionUpdated.Broadcast(SessionUpdatedEvent);
-
-				// TODO: Will be refactored in the UpdateSession split
+				FSessionUpdated SessionUpdatedEvent{ OpParams.SessionName, SessionUpdateData };
+				SessionEvents.OnSessionUpdated.Broadcast(SessionUpdatedEvent);
 			}
 			else
 			{
@@ -805,7 +777,7 @@ TOnlineAsyncOpHandle<FUpdateSession> FSessionsEOSGS::UpdateSession(FUpdateSessio
 	return Op->GetHandle();
 }
 
-TFuture<TDefaultErrorResult<FUpdateSessionImpl>> FSessionsEOSGS::UpdateSessionImpl(FUpdateSessionImpl::Params&& Params)
+TFuture<TDefaultErrorResult<FUpdateSessionImplEOSGS>> FSessionsEOSGS::UpdateSessionImplEOSGS(FUpdateSessionImplEOSGS::Params&& Params)
 {
 	EOS_Sessions_UpdateSessionOptions UpdateSessionOptions = {};
 	UpdateSessionOptions.ApiVersion = EOS_SESSIONS_UPDATESESSION_API_LATEST;
@@ -813,8 +785,8 @@ TFuture<TDefaultErrorResult<FUpdateSessionImpl>> FSessionsEOSGS::UpdateSessionIm
 
 	UpdateSessionOptions.SessionModificationHandle = Params.SessionModificationHandle->ModificationHandle;
 
-	TPromise<TDefaultErrorResult<FUpdateSessionImpl>> Promise;
-	TFuture<TDefaultErrorResult<FUpdateSessionImpl>> Future = Promise.GetFuture();
+	TPromise<TDefaultErrorResult<FUpdateSessionImplEOSGS>> Promise;
+	TFuture<TDefaultErrorResult<FUpdateSessionImplEOSGS>> Future = Promise.GetFuture();
 
 	EOS_Async(EOS_Sessions_UpdateSession, SessionsHandle, UpdateSessionOptions,
 	[this, Promise = MoveTemp(Promise), Params = MoveTemp(Params)](const EOS_Sessions_UpdateSessionCallbackInfo* Result) mutable
@@ -835,7 +807,7 @@ TFuture<TDefaultErrorResult<FUpdateSessionImpl>> FSessionsEOSGS::UpdateSessionIm
 			{
 				if (Result.IsOk())
 				{
-					Promise.EmplaceValue(FUpdateSessionImpl::Result{ NewSessionId });
+					Promise.EmplaceValue(FUpdateSessionImplEOSGS::Result{ NewSessionId });
 				}
 				else
 				{
@@ -845,7 +817,7 @@ TFuture<TDefaultErrorResult<FUpdateSessionImpl>> FSessionsEOSGS::UpdateSessionIm
 		}
 		else
 		{
-			Promise.EmplaceValue(FUpdateSessionImpl::Result{ FString(Result->SessionId) });
+			Promise.EmplaceValue(FUpdateSessionImplEOSGS::Result{ FString(Result->SessionId) });
 		}
 	});
 
diff --git a/Engine/Plugins/Online/Experimental/OnlineServicesEOSGS/Source/Public/Online/SessionsEOSGS.h b/Engine/Plugins/Online/Experimental/OnlineServicesEOSGS/Source/Public/Online/SessionsEOSGS.h
--- a/Engine/Plugins/Online/Experimental/OnlineServicesEOSGS/Source/Public/Online/SessionsEOSGS.h
+++ b/Engine/Plugins/Online/Experimental/OnlineServicesEOSGS/Source/Public/Online/SessionsEOSGS.h
@@ -114,9 +114,9 @@ struct FUpdateSessionJoinabilityImpl
 
 typedef FUpdateSessionJoinabilityImpl::Params FUpdateSessionJoinabilityParams;
 
-struct FUpdateSessionImpl
+struct FUpdateSessionImplEOSGS
 {
-	static constexpr TCHAR Name[] = TEXT("UpdateSessionImpl");
+	static constexpr TCHAR Name[] = TEXT("UpdateSessionImplEOSGS");
 
 	struct Params
 	{
@@ -191,7 +191,7 @@ public:
 
 	// ISessions
 	virtual TOnlineAsyncOpHandle<FCreateSession> CreateSession(FCreateSession::Params&& Params) override;
-	virtual TOnlineAsyncOpHandle<FUpdateSession> UpdateSession(FUpdateSession::Params&& Params) override;
+	virtual TOnlineAsyncOpHandle<FUpdateSessionImpl> UpdateSessionImpl(FUpdateSessionImpl::Params&& Params) override;
 	virtual TOnlineAsyncOpHandle<FLeaveSession> LeaveSession(FLeaveSession::Params&& Params) override;
 	virtual TOnlineAsyncOpHandle<FFindSessions> FindSessions(FFindSessions::Params&& Params) override;
 	virtual TOnlineAsyncOpHandle<FJoinSession> JoinSession(FJoinSession::Params&& Params) override;
@@ -226,9 +226,14 @@ protected:
 	void WriteCreateSessionModificationHandle(EOS_HSessionModification& SessionModificationHandle, const FCreateSession::Params& Params);
 
 	/**
-	 * Writes only the new values for all updated settings to the SessionModificationHandle
+	 * Writes only the new values for all updated session settings to the SessionModificationHandle
 	 */
-	void WriteUpdateSessionModificationHandle(EOS_HSessionModification& SessionModificationHandle, const FName& SessionName, const FSessionSettingsUpdate& NewSettings);
+	void WriteUpdateSessionModificationHandle(EOS_HSessionModification& SessionModificationHandle, const FSessionSettingsUpdate& NewSettings);
+
+	/**
+	 * Writes only the new values for all updated member settings to the SessionModificationHandle
+	 */
+	void WriteUpdateSessionModificationHandle(EOS_HSessionModification& SessionModificationHandle, const FAccountId& LocalAccountId, const FSessionMemberUpdate& MemberUpdate);
 
 	/**
 	 * Writes all relevant values set in the FindSessions parameters into the SessionSearchHandle
@@ -238,7 +243,7 @@ protected:
 	/**
 	 * Internal method used by both CreateSession and UpdateSession to process a session update at the API level
 	 */
-	TFuture<TDefaultErrorResult<FUpdateSessionImpl>> UpdateSessionImpl(FUpdateSessionImpl::Params&& Params);
+	TFuture<TDefaultErrorResult<FUpdateSessionImplEOSGS>> UpdateSessionImplEOSGS(FUpdateSessionImplEOSGS::Params&& Params);
 
 	/**
 	 * Internal method called after UpdateSessionImpl to update joinability options for a session
@@ -293,12 +298,12 @@ END_ONLINE_STRUCT_META()
 BEGIN_ONLINE_STRUCT_META(FUpdateSessionJoinabilityImpl::Result)
 END_ONLINE_STRUCT_META()		
 		
-BEGIN_ONLINE_STRUCT_META(FUpdateSessionImpl::Params)
-	ONLINE_STRUCT_FIELD(FUpdateSessionImpl::Params, SessionModificationHandle)
+BEGIN_ONLINE_STRUCT_META(FUpdateSessionImplEOSGS::Params)
+	ONLINE_STRUCT_FIELD(FUpdateSessionImplEOSGS::Params, SessionModificationHandle)
 END_ONLINE_STRUCT_META()
 
-BEGIN_ONLINE_STRUCT_META(FUpdateSessionImpl::Result)
-	ONLINE_STRUCT_FIELD(FUpdateSessionImpl::Result, NewSessionId)
+BEGIN_ONLINE_STRUCT_META(FUpdateSessionImplEOSGS::Result)
+	ONLINE_STRUCT_FIELD(FUpdateSessionImplEOSGS::Result, NewSessionId)
 END_ONLINE_STRUCT_META()
 
 BEGIN_ONLINE_STRUCT_META(FSendSingleSessionInviteImpl::Params)
diff --git a/Engine/Plugins/Online/Experimental/OnlineServicesOSSAdapter/Source/Private/Online/SessionsOSSAdapter.cpp b/Engine/Plugins/Online/Experimental/OnlineServicesOSSAdapter/Source/Private/Online/SessionsOSSAdapter.cpp
--- a/Engine/Plugins/Online/Experimental/OnlineServicesOSSAdapter/Source/Private/Online/SessionsOSSAdapter.cpp
+++ b/Engine/Plugins/Online/Experimental/OnlineServicesOSSAdapter/Source/Private/Online/SessionsOSSAdapter.cpp
@@ -505,26 +505,19 @@ TOnlineAsyncOpHandle<FCreateSession> FSessionsOSSAdapter::CreateSession(FCreateS
 	return Op->GetHandle();
 }
 
-TOnlineAsyncOpHandle<FUpdateSession> FSessionsOSSAdapter::UpdateSession(FUpdateSession::Params&& Params)
+TOnlineAsyncOpHandle<FUpdateSessionImpl> FSessionsOSSAdapter::UpdateSessionImpl(FUpdateSessionImpl::Params&& Params)
 {
-	TOnlineAsyncOpRef<FUpdateSession> Op = GetOp<FUpdateSession>(MoveTemp(Params));
-	const FUpdateSession::Params& OpParams = Op->GetParams();
+	TOnlineAsyncOpRef<FUpdateSessionImpl> Op = GetOp<FUpdateSessionImpl>(MoveTemp(Params));
+	const FUpdateSessionImpl::Params& OpParams = Op->GetParams();
 
-	Op->Then([this](TOnlineAsyncOp<FUpdateSession>& Op)
+	Op->Then([this](TOnlineAsyncOp<FUpdateSessionImpl>& Op)
 	{
-		const FUpdateSession::Params& OpParams = Op.GetParams();
-
-		FOnlineError StateCheck = CheckUpdateSessionState(OpParams);
-		if (StateCheck != Errors::Success())
-		{
-			Op.SetError(MoveTemp(StateCheck));
-			return;
-		}
+		const FUpdateSessionImpl::Params& OpParams = Op.GetParams();
 
 		MakeMulticastAdapter(this, SessionsInterface->OnUpdateSessionCompleteDelegates,
 		[this, WeakOp = Op.AsWeak()](FName SessionName, const bool bWasSuccessful)
 		{
-			if (TOnlineAsyncOpPtr<FUpdateSession> Op = WeakOp.Pin())
+			if (TOnlineAsyncOpPtr<FUpdateSessionImpl> Op = WeakOp.Pin())
 			{
 				if (!bWasSuccessful)
 				{
@@ -532,23 +525,23 @@ TOnlineAsyncOpHandle<FUpdateSession> FSessionsOSSAdapter::UpdateSession(FUpdateS
 					return;
 				}
 
-				const FUpdateSession::Params& OpParams = Op->GetParams();
+				const FUpdateSessionImpl::Params& OpParams = Op->GetParams();
 
-				const TOnlineResult<FGetMutableSessionByName> GetMutableSessionByNameResult = GetMutableSessionByName({ SessionName });
+				TOnlineResult<FGetMutableSessionByName> GetMutableSessionByNameResult = GetMutableSessionByName({ SessionName });
 				if (GetMutableSessionByNameResult.IsOk())
 				{
-					const TSharedRef<FSessionCommon>& FoundSession = GetMutableSessionByNameResult.GetOkValue().Session;
-
-					FSessionSettings& UpdatedV2Settings = FoundSession->SessionSettings;
-					//UpdatedV2Settings += OpParams.Mutations;
+					// We update our local session
+					TSharedRef<FSessionCommon>& FoundSession = GetMutableSessionByNameResult.GetOkValue().Session;
 
-					Op->SetResult({ });
+					FSessionUpdate SessionUpdateData = BuildSessionUpdate(FoundSession, OpParams.Mutations);
 
-					//FSessionUpdated Event { SessionName, OpParams.Mutations };
+					(*FoundSession) += SessionUpdateData;
 
-					//SessionEvents.OnSessionUpdated.Broadcast(Event);
+					// We set the result and fire the event
+					Op->SetResult({ });
 
-					// TODO: Will be refactored as part of the UpdateSession split
+					FSessionUpdated SessionUpdatedEvent{ OpParams.SessionName, SessionUpdateData };
+					SessionEvents.OnSessionUpdated.Broadcast(SessionUpdatedEvent);
 				}
 				else
 				{
diff --git a/Engine/Plugins/Online/Experimental/OnlineServicesOSSAdapter/Source/Public/Online/SessionsOSSAdapter.h b/Engine/Plugins/Online/Experimental/OnlineServicesOSSAdapter/Source/Public/Online/SessionsOSSAdapter.h
--- a/Engine/Plugins/Online/Experimental/OnlineServicesOSSAdapter/Source/Public/Online/SessionsOSSAdapter.h
+++ b/Engine/Plugins/Online/Experimental/OnlineServicesOSSAdapter/Source/Public/Online/SessionsOSSAdapter.h
@@ -64,7 +64,7 @@ public:
 
 	// ISessions
 	virtual TOnlineAsyncOpHandle<FCreateSession> CreateSession(FCreateSession::Params&& Params) override;
-	virtual TOnlineAsyncOpHandle<FUpdateSession> UpdateSession(FUpdateSession::Params&& Params) override;
+	virtual TOnlineAsyncOpHandle<FUpdateSessionImpl> UpdateSessionImpl(FUpdateSessionImpl::Params&& Params) override;
 	virtual TOnlineAsyncOpHandle<FLeaveSession> LeaveSession(FLeaveSession::Params&& Params) override;
 	virtual TOnlineAsyncOpHandle<FFindSessions> FindSessions(FFindSessions::Params&& Params) override;
 	virtual TOnlineAsyncOpHandle<FStartMatchmaking> StartMatchmaking(FStartMatchmaking::Params&& Params) override;
