Refactor of GetAllSessions to not error when a user calls it without having created any sessions.
Refactor of GetSessionById to not take a LocalAccountId (not needed)
Refactor of UpdateSessionImpl to return a Future instead, to have all async logic work properly.
Addition of comments for missing error checks.

diff --git a/Engine/Plugins/Online/Experimental/OnlineServices/Source/OnlineServicesCommon/Private/Online/SessionsCommon.cpp b/Engine/Plugins/Online/Experimental/OnlineServices/Source/OnlineServicesCommon/Private/Online/SessionsCommon.cpp
--- a/Engine/Plugins/Online/Experimental/OnlineServices/Source/OnlineServicesCommon/Private/Online/SessionsCommon.cpp
+++ b/Engine/Plugins/Online/Experimental/OnlineServices/Source/OnlineServicesCommon/Private/Online/SessionsCommon.cpp
@@ -82,6 +82,10 @@ namespace UE::Online {
 		RegisterCommand(&FSessionsCommon::GetAllSessions);
 		RegisterCommand(&FSessionsCommon::GetSessionByName);
 		RegisterCommand(&FSessionsCommon::GetSessionById);
+		RegisterCommand(&FSessionsCommon::GetPresenceSession);
+		RegisterCommand(&FSessionsCommon::IsPresenceSession);
+		RegisterCommand(&FSessionsCommon::SetPresenceSession);
+		RegisterCommand(&FSessionsCommon::ClearPresenceSession);
 		RegisterCommand(&FSessionsCommon::CreateSession);
 		RegisterCommand(&FSessionsCommon::UpdateSessionSettings);
 		RegisterCommand(&FSessionsCommon::UpdateSessionMember);
@@ -89,29 +93,33 @@ namespace UE::Online {
 		RegisterCommand(&FSessionsCommon::FindSessions);
 		RegisterCommand(&FSessionsCommon::StartMatchmaking);
 		RegisterCommand(&FSessionsCommon::JoinSession);
+		RegisterCommand(&FSessionsCommon::AddSessionMember);
+		RegisterCommand(&FSessionsCommon::RemoveSessionMember);
 		RegisterCommand(&FSessionsCommon::SendSessionInvite);
 		RegisterCommand(&FSessionsCommon::RejectSessionInvite);
 	}
 
 	TOnlineResult<FGetAllSessions> FSessionsCommon::GetAllSessions(FGetAllSessions::Params&& Params) const
 	{
+		// TODO: Params and user login check
+
 		if (const TArray<FName>* UserSessions = NamedSessionUserMap.Find(Params.LocalAccountId))
 		{
-		FGetAllSessions::Result Result;
+			FGetAllSessions::Result Result;
 
 			for (const FName& SessionName : *UserSessions)
-		{
+			{
 				const FOnlineSessionId& SessionId = LocalSessionsByName.FindChecked(SessionName);
 				const TSharedRef<FSessionCommon>& Session = AllSessionsById.FindChecked(SessionId);
 
 				Result.Sessions.Add(Session);
-		}
+			}
 
-		return TOnlineResult<FGetAllSessions>(MoveTemp(Result));
-	}
+			return TOnlineResult<FGetAllSessions>(MoveTemp(Result));
+		}
 		else
 		{
-			return TOnlineResult<FGetAllSessions>(Errors::InvalidUser());
+			return TOnlineResult<FGetAllSessions>(Errors::NotFound());
 		}
 	}
 
@@ -143,6 +151,8 @@ namespace UE::Online {
 
 	TOnlineResult<FGetPresenceSession> FSessionsCommon::GetPresenceSession(FGetPresenceSession::Params&& Params) const
 	{
+		// TODO: Check the user is valid and logged in
+
 		if (const FOnlineSessionId* PresenceSessionId = PresenceSessionsUserMap.Find(Params.LocalAccountId))
 		{
 			return TOnlineResult<FGetPresenceSession>({ AllSessionsById.FindChecked(*PresenceSessionId) });
@@ -155,6 +165,8 @@ namespace UE::Online {
 
 	TOnlineResult<FIsPresenceSession> FSessionsCommon::IsPresenceSession(FIsPresenceSession::Params&& Params) const
 	{
+		// TODO: Check the user is valid and logged in. Check the session id is valid
+
 		if (const FOnlineSessionId* PresenceSessionId = PresenceSessionsUserMap.Find(Params.LocalAccountId))
 		{
 			return TOnlineResult<FIsPresenceSession>(FIsPresenceSession::Result{ Params.SessionId == (*PresenceSessionId) });
@@ -167,6 +179,8 @@ namespace UE::Online {
 
 	TOnlineResult<FSetPresenceSession> FSessionsCommon::SetPresenceSession(FSetPresenceSession::Params&& Params)
 	{
+		// TODO: Check the user is valid and logged in. Check the session id is valid
+
 		FOnlineSessionId& PresenceSessionId = PresenceSessionsUserMap.FindOrAdd(Params.LocalAccountId);
 		PresenceSessionId = Params.SessionId;
 
@@ -175,6 +189,8 @@ namespace UE::Online {
 
 	TOnlineResult<FClearPresenceSession> FSessionsCommon::ClearPresenceSession(FClearPresenceSession::Params&& Params)
 	{
+		// TODO: Check the user is valid and logged in
+
 		PresenceSessionsUserMap.Remove(Params.LocalAccountId);
 
 		return TOnlineResult<FClearPresenceSession>(FClearPresenceSession::Result{ });
@@ -211,7 +227,7 @@ namespace UE::Online {
 			}
 
 			UpdateSessionImpl({ OpParams.LocalAccountId, OpParams.SessionName, { OpParams.Mutations, {} } })
-			.OnComplete([this, WeakOp = Op.AsWeak()](const TOnlineResult<FUpdateSessionImpl>& Result)
+			.Next([this, WeakOp = Op.AsWeak()](const TOnlineResult<FUpdateSessionImpl>& Result)
 			{
 				if (TOnlineAsyncOpPtr<FUpdateSessionSettings> StrongOp = WeakOp.Pin())
 				{
@@ -258,7 +274,7 @@ namespace UE::Online {
 			}
 
 			UpdateSessionImpl({ OpParams.LocalAccountId, OpParams.SessionName, { FSessionSettingsUpdate() , { { OpParams.LocalAccountId, OpParams.Mutations } } } })
-			.OnComplete([this, WeakOp = Op.AsWeak()](const TOnlineResult<FUpdateSessionImpl>& Result)
+			.Next([this, WeakOp = Op.AsWeak()](const TOnlineResult<FUpdateSessionImpl>& Result)
 			{
 				if (TOnlineAsyncOpPtr<FUpdateSessionMember> StrongOp = WeakOp.Pin())
 				{
@@ -427,6 +443,8 @@ namespace UE::Online {
 
 	TOnlineResult<FGetMutableSessionByName> FSessionsCommon::GetMutableSessionByName(FGetMutableSessionByName::Params&& Params) const
 	{
+		// TODO: Check that name is not empty
+
 		if (const FOnlineSessionId* SessionId = LocalSessionsByName.Find(Params.LocalName))
 		{
 			check(AllSessionsById.Contains(*SessionId));
@@ -441,6 +459,8 @@ namespace UE::Online {
 
 	TOnlineResult<FGetMutableSessionById> FSessionsCommon::GetMutableSessionById(FGetMutableSessionById::Params&& Params) const
 	{
+		// TODO: Check that session id is valid
+
 		if (const TSharedRef<FSessionCommon>* FoundSession = AllSessionsById.Find(Params.SessionId))
 		{
 			return TOnlineResult<FGetMutableSessionById>({ *FoundSession });
@@ -643,11 +663,9 @@ namespace UE::Online {
 		return Result;
 	}
 
-	TOnlineAsyncOpHandle<FUpdateSessionImpl> FSessionsCommon::UpdateSessionImpl(FUpdateSessionImpl::Params&& Params)
+	TFuture<TOnlineResult<FUpdateSessionImpl>> FSessionsCommon::UpdateSessionImpl(FUpdateSessionImpl::Params&& Params)
 	{
-		TOnlineAsyncOpRef<FUpdateSessionImpl> Operation = GetOp<FUpdateSessionImpl>(MoveTemp(Params));
-		Operation->SetError(Errors::NotImplemented());
-		return Operation->GetHandle();
+		return MakeFulfilledPromise<TOnlineResult<FUpdateSessionImpl>>(Errors::NotImplemented()).GetFuture();
 	}
 
 	TOnlineResult<FAddSessionMember> FSessionsCommon::AddSessionMemberImpl(const FAddSessionMember::Params& Params)
@@ -721,7 +739,7 @@ namespace UE::Online {
 	FOnlineError FSessionsCommon::CheckCreateSessionState(const FCreateSession::Params& Params)
 	{
 		TOnlineResult<FGetSessionByName> GetSessionByNameResult = GetSessionByName({ Params.SessionName });
-		if (GetSessionByNameResult.IsError())
+		if (GetSessionByNameResult.IsOk())
 		{
 			UE_LOG(LogTemp, Warning, TEXT("[FSessionsCommon::CheckCreateSessionState] Could not create session with name [%s]. A session with that name already exists"), *Params.SessionName.ToString());
 
@@ -961,7 +979,7 @@ namespace UE::Online {
 		}
 
 		// We check that the session is cached and valid for a join operation by the users
-		TOnlineResult<FGetSessionById> GetSessionByIdResult = GetSessionById({ Params.LocalAccountId, Params.SessionId });
+		TOnlineResult<FGetSessionById> GetSessionByIdResult = GetSessionById({ Params.SessionId });
 		if (GetSessionByIdResult.IsError())
 		{
 			UE_LOG(LogTemp, Warning, TEXT("[FSessionsCommon::CheckJoinSessionState] Unable to find session with SessionId parameter [%s]. Please call FindSessions to get an updated list of available sessions "), *ToLogString(Params.SessionId));
diff --git a/Engine/Plugins/Online/Experimental/OnlineServices/Source/OnlineServicesCommon/Public/Online/SessionsCommon.h b/Engine/Plugins/Online/Experimental/OnlineServices/Source/OnlineServicesCommon/Public/Online/SessionsCommon.h
--- a/Engine/Plugins/Online/Experimental/OnlineServices/Source/OnlineServicesCommon/Public/Online/SessionsCommon.h
+++ b/Engine/Plugins/Online/Experimental/OnlineServices/Source/OnlineServicesCommon/Public/Online/SessionsCommon.h
@@ -144,7 +144,7 @@ public:
 
 	virtual bool IsJoinable() const override								{ return GetNumOpenConnections() > 0 && SessionSettings.bAllowNewMembers; }
 
-	virtual FString ToLogString() const override							{ return FString(); } // TODO: Implement after completing refactor
+	virtual FString ToLogString() const override							{ return FString::Printf(TEXT("[FSessionCommon: SessionId[%d]]"), SessionInfo.SessionId.GetHandle()); } // TODO: Expand with relevant information
 
 public:
 	/** Session information that will remain constant throughout the session's lifetime */
@@ -284,7 +284,7 @@ protected:
 	FSessionSettingsChanges BuildSessionSettingsChanges(const TSharedRef<FSessionCommon>& Session, const FSessionSettingsUpdate& UpdatedValues) const;
 	FSessionUpdate BuildSessionUpdate(const TSharedRef<FSessionCommon>& Session, const FCombinedSessionUpdate& UpdatedValues) const;
 
-	virtual TOnlineAsyncOpHandle<FUpdateSessionImpl> UpdateSessionImpl(FUpdateSessionImpl::Params&& Params);
+	virtual TFuture<TOnlineResult<FUpdateSessionImpl>> UpdateSessionImpl(FUpdateSessionImpl::Params&& Params);
 	TOnlineResult<FAddSessionMember> AddSessionMemberImpl(const FAddSessionMember::Params& Params);
 	TOnlineResult<FRemoveSessionMember> RemoveSessionMemberImpl(const FRemoveSessionMember::Params& Params);
 
diff --git a/Engine/Plugins/Online/Experimental/OnlineServices/Source/OnlineServicesCommonEngineUtils/Private/Online/SessionsLAN.cpp b/Engine/Plugins/Online/Experimental/OnlineServices/Source/OnlineServicesCommonEngineUtils/Private/Online/SessionsLAN.cpp
--- a/Engine/Plugins/Online/Experimental/OnlineServices/Source/OnlineServicesCommonEngineUtils/Private/Online/SessionsLAN.cpp
+++ b/Engine/Plugins/Online/Experimental/OnlineServices/Source/OnlineServicesCommonEngineUtils/Private/Online/SessionsLAN.cpp
@@ -153,7 +153,12 @@ TOnlineAsyncOpHandle<FCreateSession> FSessionsLAN::CreateSession(FCreateSession:
 
 		TSharedRef<FSessionLAN> NewSessionLANRef = MakeShared<FSessionLAN>();
 		NewSessionLANRef->OwnerAccountId = OpParams.LocalAccountId;
+		NewSessionLANRef->SessionInfo.bAllowSanctionedPlayers = OpParams.bAllowSanctionedPlayers;
+		NewSessionLANRef->SessionInfo.bAntiCheatProtected = OpParams.bAntiCheatProtected;
+		NewSessionLANRef->SessionInfo.bIsDedicatedServerSession = IsRunningDedicatedServer();
+		NewSessionLANRef->SessionInfo.bIsLANSession = true;
 		NewSessionLANRef->SessionInfo.SessionId = FOnlineSessionIdRegistryLAN::GetChecked(Services.GetServicesProvider()).GetNextSessionId();
+		NewSessionLANRef->SessionInfo.SessionIdOverride = OpParams.SessionIdOverride;
 		NewSessionLANRef->SessionSettings = OpParams.SessionSettings;
 
 		// For LAN sessions, we'll add the Session member manually instead of calling JoinSession since there is no API calls involved
@@ -169,55 +174,50 @@ TOnlineAsyncOpHandle<FCreateSession> FSessionsLAN::CreateSession(FCreateSession:
 	return Op->GetHandle();
 }
 
-TOnlineAsyncOpHandle<FUpdateSessionImpl> FSessionsLAN::UpdateSessionImpl(FUpdateSessionImpl::Params&& Params)
+TFuture<TOnlineResult<FUpdateSessionImpl>> FSessionsLAN::UpdateSessionImpl(FUpdateSessionImpl::Params&& Params)
 {
-	TOnlineAsyncOpRef<FUpdateSessionImpl> Op = GetOp<FUpdateSessionImpl>(MoveTemp(Params));
+	TPromise<TOnlineResult<FUpdateSessionImpl>> Promise;
+	TFuture<TOnlineResult<FUpdateSessionImpl>> Future = Promise.GetFuture();
 
-	Op->Then([this](TOnlineAsyncOp<FUpdateSessionImpl>& Op) mutable
-	{
-		const FUpdateSessionImpl::Params& OpParams = Op.GetParams();
+	TOnlineResult<FGetMutableSessionByName> GetMutableSessionByNameResult = GetMutableSessionByName({ Params.SessionName });
+	check(GetMutableSessionByNameResult.IsOk());
 
-		// We'll check and update all settings one by one, with additional logic wherever required
+	// We'll individually check only the settings where additional logic is required
 
-		TOnlineResult<FGetMutableSessionByName> GetMutableSessionByNameResult = GetMutableSessionByName({ OpParams.SessionName });
-		check(GetMutableSessionByNameResult.IsOk());
-		
-		TSharedRef<FSessionCommon> FoundSession = GetMutableSessionByNameResult.GetOkValue().Session;
-		FSessionSettings& SessionSettings = FoundSession->SessionSettings;
-		const FCombinedSessionUpdate& SessionUpdate = OpParams.Mutations;
+	TSharedRef<FSessionCommon> FoundSession = GetMutableSessionByNameResult.GetOkValue().Session;
+	FSessionSettings& SessionSettings = FoundSession->SessionSettings;
+	const FCombinedSessionUpdate& SessionUpdate = Params.Mutations;
 
-		if (SessionUpdate.UpdatedSessionSettings.JoinPolicy.IsSet())
+	if (SessionUpdate.UpdatedSessionSettings.JoinPolicy.IsSet())
+	{
+		// Changes in the join policy setting will mean we start or stop the LAN Beacon broadcasting the session information
+		// There is no FriendsLAN interface at the time of this implementation, hence the binary behavior (Public/Non-Public)
+		if (SessionSettings.JoinPolicy != ESessionJoinPolicy::Public && SessionUpdate.UpdatedSessionSettings.JoinPolicy.GetValue() == ESessionJoinPolicy::Public)
 		{
-			// Changes in the join policy setting will mean we start or stop the LAN Beacon broadcasting the session information
-			// There is no FriendsLAN interface at the time of this implementation, hence the binary behavior (Public/Non-Public)
-			if (SessionSettings.JoinPolicy != ESessionJoinPolicy::Public && SessionUpdate.UpdatedSessionSettings.JoinPolicy.GetValue() == ESessionJoinPolicy::Public)
-			{
-				if (TOptional<FOnlineError> TryHostLANSessionResult = TryHostLANSession())
-				{
-					Op.SetError(MoveTemp(TryHostLANSessionResult.GetValue()));
-					return;
-				}
-			}
-			else if (SessionSettings.JoinPolicy == ESessionJoinPolicy::Public && SessionUpdate.UpdatedSessionSettings.JoinPolicy.GetValue() != ESessionJoinPolicy::Public)
+			if (TOptional<FOnlineError> TryHostLANSessionResult = TryHostLANSession())
 			{
-				StopLANSession();
+				Promise.EmplaceValue(TryHostLANSessionResult.GetValue());
+				return Future;
 			}
 		}
+		else if (SessionSettings.JoinPolicy == ESessionJoinPolicy::Public && SessionUpdate.UpdatedSessionSettings.JoinPolicy.GetValue() != ESessionJoinPolicy::Public)
+		{
+			StopLANSession();
+		}
+	}
 
-		// We update our local session
-		FSessionUpdate SessionUpdateData = BuildSessionUpdate(FoundSession, OpParams.Mutations);
+	// We update our local session
+	FSessionUpdate SessionUpdateData = BuildSessionUpdate(FoundSession, Params.Mutations);
 
-		(*FoundSession) += SessionUpdateData;
+	(*FoundSession) += SessionUpdateData;
 
-		// We set the result and fire the event
-		Op.SetResult({ });
+	// We set the result and fire the event
+	Promise.EmplaceValue(FUpdateSessionImpl::Result{ });
 
-		FSessionUpdated SessionUpdatedEvent{ OpParams.SessionName, SessionUpdateData };
-		SessionEvents.OnSessionUpdated.Broadcast(SessionUpdatedEvent);
-	})
-	.Enqueue(GetSerialQueue());
+	FSessionUpdated SessionUpdatedEvent{ Params.SessionName, SessionUpdateData };
+	SessionEvents.OnSessionUpdated.Broadcast(SessionUpdatedEvent);
 
-	return Op->GetHandle();
+	return Future;
 }
 
 TOnlineAsyncOpHandle<FFindSessions> FSessionsLAN::FindSessions(FFindSessions::Params&& Params)
@@ -276,7 +276,7 @@ TOnlineAsyncOpHandle<FJoinSession> FSessionsLAN::JoinSession(FJoinSession::Param
 			return;
 		}
 
-		TOnlineResult<FGetSessionById> GetSessionByIdResult = GetSessionById({ OpParams.LocalAccountId, OpParams.SessionId });
+		TOnlineResult<FGetSessionById> GetSessionByIdResult = GetSessionById({ OpParams.SessionId });
 		if (GetSessionByIdResult.IsError())
 		{
 			// If no result is found, the id might be expired, which we should notify
diff --git a/Engine/Plugins/Online/Experimental/OnlineServices/Source/OnlineServicesCommonEngineUtils/Public/Online/SessionsLAN.h b/Engine/Plugins/Online/Experimental/OnlineServices/Source/OnlineServicesCommonEngineUtils/Public/Online/SessionsLAN.h
--- a/Engine/Plugins/Online/Experimental/OnlineServices/Source/OnlineServicesCommonEngineUtils/Public/Online/SessionsLAN.h
+++ b/Engine/Plugins/Online/Experimental/OnlineServices/Source/OnlineServicesCommonEngineUtils/Public/Online/SessionsLAN.h
@@ -57,11 +57,13 @@ public:
 
 	// ISessions
 	virtual TOnlineAsyncOpHandle<FCreateSession> CreateSession(FCreateSession::Params&& Params) override;
-	virtual TOnlineAsyncOpHandle<FUpdateSessionImpl> UpdateSessionImpl(FUpdateSessionImpl::Params&& Params) override;
 	virtual TOnlineAsyncOpHandle<FFindSessions> FindSessions(FFindSessions::Params&& Params) override;
 	virtual TOnlineAsyncOpHandle<FJoinSession> JoinSession(FJoinSession::Params&& Params) override;
 	virtual TOnlineAsyncOpHandle<FLeaveSession> LeaveSession(FLeaveSession::Params&& Params) override;
 
+	// FSessionsCommon
+	virtual TFuture<TOnlineResult<FUpdateSessionImpl>> UpdateSessionImpl(FUpdateSessionImpl::Params&& Params) override;
+
 protected:
 	TOptional<FOnlineError> TryHostLANSession();
 	void FindLANSessions(const FAccountId& LocalAccountId);
diff --git a/Engine/Plugins/Online/Experimental/OnlineServices/Source/OnlineServicesInterface/Public/Online/Sessions.h b/Engine/Plugins/Online/Experimental/OnlineServices/Source/OnlineServicesInterface/Public/Online/Sessions.h
--- a/Engine/Plugins/Online/Experimental/OnlineServices/Source/OnlineServicesInterface/Public/Online/Sessions.h
+++ b/Engine/Plugins/Online/Experimental/OnlineServices/Source/OnlineServicesInterface/Public/Online/Sessions.h
@@ -252,8 +252,6 @@ struct FGetSessionById
 
 	struct Params
 	{
-		FAccountId LocalAccountId;
-
 		FOnlineSessionId SessionId;
 	};
 
@@ -276,7 +274,7 @@ struct FGetPresenceSession
 
 	struct Result
 	{
-		TSharedPtr<const ISession> Session;
+		TSharedRef<const ISession> Session;
 	};
 };
 
@@ -975,7 +973,6 @@ BEGIN_ONLINE_STRUCT_META(FGetSessionByName::Result)
 END_ONLINE_STRUCT_META()
 
 BEGIN_ONLINE_STRUCT_META(FGetSessionById::Params)
-	ONLINE_STRUCT_FIELD(FGetSessionById::Params, LocalAccountId),
 	ONLINE_STRUCT_FIELD(FGetSessionById::Params, SessionId)
 END_ONLINE_STRUCT_META()
 
diff --git a/Engine/Plugins/Online/Experimental/OnlineServicesEOSGS/Source/Private/Online/OnlineServicesEOSGS.cpp b/Engine/Plugins/Online/Experimental/OnlineServicesEOSGS/Source/Private/Online/OnlineServicesEOSGS.cpp
--- a/Engine/Plugins/Online/Experimental/OnlineServicesEOSGS/Source/Private/Online/OnlineServicesEOSGS.cpp
+++ b/Engine/Plugins/Online/Experimental/OnlineServicesEOSGS/Source/Private/Online/OnlineServicesEOSGS.cpp
@@ -130,7 +130,7 @@ TOnlineResult<FGetResolvedConnectString> FOnlineServicesEOSGS::GetResolvedConnec
 		ISessionsPtr SessionsEOS = GetSessionsInterface();
 		check(SessionsEOS);
 
-		TOnlineResult<FGetSessionById> Result = SessionsEOS->GetSessionById({ Params.LocalAccountId, Params.SessionId });
+		TOnlineResult<FGetSessionById> Result = SessionsEOS->GetSessionById({ Params.SessionId });
 		if (Result.IsOk())
 		{
 #if WITH_ENGINE
diff --git a/Engine/Plugins/Online/Experimental/OnlineServicesEOSGS/Source/Private/Online/SessionsEOSGS.cpp b/Engine/Plugins/Online/Experimental/OnlineServicesEOSGS/Source/Private/Online/SessionsEOSGS.cpp
--- a/Engine/Plugins/Online/Experimental/OnlineServicesEOSGS/Source/Private/Online/SessionsEOSGS.cpp
+++ b/Engine/Plugins/Online/Experimental/OnlineServicesEOSGS/Source/Private/Online/SessionsEOSGS.cpp
@@ -685,7 +685,7 @@ void FSessionsEOSGS::WriteUpdateSessionModificationHandle(EOS_HSessionModificati
 	}
 }
 
-TOnlineAsyncOpHandle<FUpdateSessionImpl> FSessionsEOSGS::UpdateSessionImpl(FUpdateSessionImpl::Params&& Params)
+TFuture<TOnlineResult<FUpdateSessionImpl>> FSessionsEOSGS::UpdateSessionImpl(FUpdateSessionImpl::Params&& Params)
 {
 	// LAN Sessions
 	TOnlineResult<FGetSessionByName> Result = GetSessionByName({ Params.SessionName });
@@ -696,56 +696,49 @@ TOnlineAsyncOpHandle<FUpdateSessionImpl> FSessionsEOSGS::UpdateSessionImpl(FUpda
 
 	// EOSGS Sessions
 
-	TOnlineAsyncOpRef<FUpdateSessionImpl> Op = GetOp<FUpdateSessionImpl>(MoveTemp(Params));
-	const FUpdateSessionImpl::Params& OpParams = Op->GetParams();
+	TPromise<TOnlineResult<FUpdateSessionImpl>> Promise;
+	TFuture<TOnlineResult<FUpdateSessionImpl>> Future = Promise.GetFuture();
 
-	Op->Then([this](TOnlineAsyncOp<FUpdateSessionImpl>& Op)
-	{
-		const FUpdateSessionImpl::Params& OpParams = Op.GetParams();
-
-		EOS_Sessions_UpdateSessionModificationOptions UpdateSessionModificationOptions = {};
-		UpdateSessionModificationOptions.ApiVersion = EOS_SESSIONS_UPDATESESSIONMODIFICATION_API_LATEST;
-		static_assert(EOS_SESSIONS_UPDATESESSIONMODIFICATION_API_LATEST == 1, "EOS_Sessions_UpdateSessionModificationOptions updated, check new fields");
+	EOS_Sessions_UpdateSessionModificationOptions UpdateSessionModificationOptions = {};
+	UpdateSessionModificationOptions.ApiVersion = EOS_SESSIONS_UPDATESESSIONMODIFICATION_API_LATEST;
+	static_assert(EOS_SESSIONS_UPDATESESSIONMODIFICATION_API_LATEST == 1, "EOS_Sessions_UpdateSessionModificationOptions updated, check new fields");
 
-		const FTCHARToUTF8 SessionNameUtf8(*OpParams.SessionName.ToString());
-		UpdateSessionModificationOptions.SessionName = SessionNameUtf8.Get();
+	const FTCHARToUTF8 SessionNameUtf8(*Params.SessionName.ToString());
+	UpdateSessionModificationOptions.SessionName = SessionNameUtf8.Get();
 
-		EOS_HSessionModification SessionModificationHandle = nullptr;
-		EOS_EResult ResultCode = EOS_Sessions_UpdateSessionModification(SessionsHandle, &UpdateSessionModificationOptions, &SessionModificationHandle);
-		if (ResultCode != EOS_EResult::EOS_Success)
-		{
-			UE_LOG(LogTemp, Error, TEXT("[FSessionsEOSGS::UpdateSession] EOS_Sessions_UpdateSessionModification failed with result [%s]"), *LexToString(ResultCode));
+	EOS_HSessionModification SessionModificationHandle = nullptr;
+	EOS_EResult ResultCode = EOS_Sessions_UpdateSessionModification(SessionsHandle, &UpdateSessionModificationOptions, &SessionModificationHandle);
+	if (ResultCode != EOS_EResult::EOS_Success)
+	{
+		UE_LOG(LogTemp, Error, TEXT("[FSessionsEOSGS::UpdateSession] EOS_Sessions_UpdateSessionModification failed with result [%s]"), *LexToString(ResultCode));
 
-			Op.SetError(Errors::FromEOSResult(ResultCode));
+		Promise.EmplaceValue(Errors::FromEOSResult(ResultCode));
 
-			return MakeFulfilledPromise<TDefaultErrorResult<FUpdateSessionImplEOSGS>>().GetFuture();
-		}
+		return Future;
+	}
 
-		// After creating the session modification handle, we'll add all the updated data to it
-		WriteUpdateSessionModificationHandle(SessionModificationHandle, OpParams.Mutations.UpdatedSessionSettings);
+	// After creating the session modification handle, we'll add all the updated data to it
+	WriteUpdateSessionModificationHandle(SessionModificationHandle, Params.Mutations.UpdatedSessionSettings);
 
-		for (const TPair<FAccountId, FSessionMemberUpdate>& Entry : OpParams.Mutations.UpdatedSessionMembers)
-		{
-			WriteUpdateSessionModificationHandle(SessionModificationHandle, Entry.Key, Entry.Value);
-		}
+	for (const TPair<FAccountId, FSessionMemberUpdate>& Entry : Params.Mutations.UpdatedSessionMembers)
+	{
+		WriteUpdateSessionModificationHandle(SessionModificationHandle, Entry.Key, Entry.Value);
+	}
 
-		// Whether we update joinability or not will depend on inf bAllowNewMembers was set to a new value
-		FUpdateSessionImplEOSGS::Params UpdateSessionImplEOSGSParams =
-		{
-			MakeShared<FSessionModificationHandleEOSGS>(SessionModificationHandle),
-			OpParams.Mutations.UpdatedSessionSettings.bAllowNewMembers.IsSet() ?
-				FUpdateSessionJoinabilityParams{ OpParams.SessionName, OpParams.Mutations.UpdatedSessionSettings.bAllowNewMembers.GetValue() } :
-				TOptional<FUpdateSessionJoinabilityParams>() // If bAllowNewMembers is not set, the TOptional will be unset too
-		};
+	// Whether we update joinability or not will depend on if bAllowNewMembers was set to a new value
+	FUpdateSessionImplEOSGS::Params UpdateSessionImplEOSGSParams =
+	{
+		MakeShared<FSessionModificationHandleEOSGS>(SessionModificationHandle),
+		Params.Mutations.UpdatedSessionSettings.bAllowNewMembers.IsSet() ?
+			FUpdateSessionJoinabilityParams{ Params.SessionName, Params.Mutations.UpdatedSessionSettings.bAllowNewMembers.GetValue() } :
+			TOptional<FUpdateSessionJoinabilityParams>() // If bAllowNewMembers is not set, the TOptional will be unset too
+	};
 
-		return UpdateSessionImplEOSGS(MoveTemp(UpdateSessionImplEOSGSParams));
-	})
-	.Then([this](TOnlineAsyncOp<FUpdateSessionImpl>& Op, TDefaultErrorResult<FUpdateSessionImplEOSGS>&& Result)
+	UpdateSessionImplEOSGS(MoveTemp(UpdateSessionImplEOSGSParams))
+	.Next([this, Promise = MoveTemp(Promise), OpParams = MoveTemp(Params)](TDefaultErrorResult<FUpdateSessionImplEOSGS>&& Result) mutable
 	{
 		if (Result.IsOk())
 		{
-			const FUpdateSessionImpl::Params& OpParams = Op.GetParams();
-
 			TOnlineResult<FGetMutableSessionByName> GetMutableSessionByNameResult = GetMutableSessionByName({ OpParams.SessionName });
 			if (GetMutableSessionByNameResult.IsOk())
 			{
@@ -757,24 +750,23 @@ TOnlineAsyncOpHandle<FUpdateSessionImpl> FSessionsEOSGS::UpdateSessionImpl(FUpda
 				(*FoundSession) += SessionUpdateData;
 
 				// We set the result and fire the event
-				Op.SetResult({ });
+				Promise.EmplaceValue(FUpdateSessionImpl::Result{ });
 
 				FSessionUpdated SessionUpdatedEvent{ OpParams.SessionName, SessionUpdateData };
 				SessionEvents.OnSessionUpdated.Broadcast(SessionUpdatedEvent);
 			}
 			else
 			{
-				Op.SetError(MoveTemp(GetMutableSessionByNameResult.GetErrorValue()));
+				Promise.EmplaceValue(GetMutableSessionByNameResult.GetErrorValue());
 			}
 		}
 		else
 		{
-			Op.SetError(MoveTemp(Result.GetErrorValue()));
+			Promise.EmplaceValue(Result.GetErrorValue());
 		}
-	})
-	.Enqueue(GetSerialQueue());
+	});
 
-	return Op->GetHandle();
+	return Future;
 }
 
 TFuture<TDefaultErrorResult<FUpdateSessionImplEOSGS>> FSessionsEOSGS::UpdateSessionImplEOSGS(FUpdateSessionImplEOSGS::Params&& Params)
@@ -1281,7 +1273,7 @@ TOnlineAsyncOpHandle<FJoinSession> FSessionsEOSGS::JoinSession(FJoinSession::Par
 	TOnlineAsyncOpRef<FJoinSession> Op = GetOp<FJoinSession>(MoveTemp(Params));
 	const FJoinSession::Params& OpParams = Op->GetParams();
 
-	TOnlineResult<FGetSessionById> GetSessionByIdResult = GetSessionById({ OpParams.LocalAccountId, OpParams.SessionId });
+	TOnlineResult<FGetSessionById> GetSessionByIdResult = GetSessionById({ OpParams.SessionId });
 	if (GetSessionByIdResult.IsError())
 	{
 		// If no result is found, the id might be expired, which we should notify
@@ -1334,7 +1326,7 @@ TOnlineAsyncOpHandle<FJoinSession> FSessionsEOSGS::JoinSession(FJoinSession::Par
 			return;
 		}
 
-		TOnlineResult<FGetSessionById> GetSessionByIdResult = GetSessionById({ OpParams.LocalAccountId, OpParams.SessionId });
+		TOnlineResult<FGetSessionById> GetSessionByIdResult = GetSessionById({ OpParams.SessionId });
 		if (GetSessionByIdResult.IsError())
 		{
 			// If no result is found, the id might be expired, which we should notify
@@ -1377,7 +1369,7 @@ TOnlineAsyncOpHandle<FJoinSession> FSessionsEOSGS::JoinSession(FJoinSession::Par
 			return;
 		}
 		
-		TOnlineResult<FGetSessionById> GetSessionByIdResult = GetSessionById({ OpParams.LocalAccountId, OpParams.SessionId });
+		TOnlineResult<FGetSessionById> GetSessionByIdResult = GetSessionById({ OpParams.SessionId });
 		if (GetSessionByIdResult.IsError())
 		{
 			// If no result is found, the id might be expired, which we should notify
diff --git a/Engine/Plugins/Online/Experimental/OnlineServicesEOSGS/Source/Public/Online/SessionsEOSGS.h b/Engine/Plugins/Online/Experimental/OnlineServicesEOSGS/Source/Public/Online/SessionsEOSGS.h
--- a/Engine/Plugins/Online/Experimental/OnlineServicesEOSGS/Source/Public/Online/SessionsEOSGS.h
+++ b/Engine/Plugins/Online/Experimental/OnlineServicesEOSGS/Source/Public/Online/SessionsEOSGS.h
@@ -191,7 +191,6 @@ public:
 
 	// ISessions
 	virtual TOnlineAsyncOpHandle<FCreateSession> CreateSession(FCreateSession::Params&& Params) override;
-	virtual TOnlineAsyncOpHandle<FUpdateSessionImpl> UpdateSessionImpl(FUpdateSessionImpl::Params&& Params) override;
 	virtual TOnlineAsyncOpHandle<FLeaveSession> LeaveSession(FLeaveSession::Params&& Params) override;
 	virtual TOnlineAsyncOpHandle<FFindSessions> FindSessions(FFindSessions::Params&& Params) override;
 	virtual TOnlineAsyncOpHandle<FJoinSession> JoinSession(FJoinSession::Params&& Params) override;
@@ -200,6 +199,9 @@ public:
 	virtual TOnlineAsyncOpHandle<FAddSessionMember> AddSessionMember(FAddSessionMember::Params&& Params) override;
 	virtual TOnlineAsyncOpHandle<FRemoveSessionMember> RemoveSessionMember(FRemoveSessionMember::Params&& Params) override;
 
+	// FSessionsCommon
+	virtual TFuture<TOnlineResult<FUpdateSessionImpl>> UpdateSessionImpl(FUpdateSessionImpl::Params&& Params) override;
+
 protected:
 	void RegisterEventHandlers();
 	void UnregisterEventHandlers();
diff --git a/Engine/Plugins/Online/Experimental/OnlineServicesOSSAdapter/Source/Private/Online/SessionsOSSAdapter.cpp b/Engine/Plugins/Online/Experimental/OnlineServicesOSSAdapter/Source/Private/Online/SessionsOSSAdapter.cpp
--- a/Engine/Plugins/Online/Experimental/OnlineServicesOSSAdapter/Source/Private/Online/SessionsOSSAdapter.cpp
+++ b/Engine/Plugins/Online/Experimental/OnlineServicesOSSAdapter/Source/Private/Online/SessionsOSSAdapter.cpp
@@ -505,69 +505,57 @@ TOnlineAsyncOpHandle<FCreateSession> FSessionsOSSAdapter::CreateSession(FCreateS
 	return Op->GetHandle();
 }
 
-TOnlineAsyncOpHandle<FUpdateSessionImpl> FSessionsOSSAdapter::UpdateSessionImpl(FUpdateSessionImpl::Params&& Params)
+TFuture<TOnlineResult<FUpdateSessionImpl>> FSessionsOSSAdapter::UpdateSessionImpl(FUpdateSessionImpl::Params&& Params)
 {
-	TOnlineAsyncOpRef<FUpdateSessionImpl> Op = GetOp<FUpdateSessionImpl>(MoveTemp(Params));
-	const FUpdateSessionImpl::Params& OpParams = Op->GetParams();
+	TPromise<TOnlineResult<FUpdateSessionImpl>> Promise;
+	TFuture<TOnlineResult<FUpdateSessionImpl>> Future = Promise.GetFuture();
 
-	Op->Then([this](TOnlineAsyncOp<FUpdateSessionImpl>& Op)
+	MakeMulticastAdapter(this, SessionsInterface->OnUpdateSessionCompleteDelegates,
+	[this, Promise = MoveTemp(Promise), OpParams = Params](FName SessionName, const bool bWasSuccessful) mutable
 	{
-		const FUpdateSessionImpl::Params& OpParams = Op.GetParams();
-
-		MakeMulticastAdapter(this, SessionsInterface->OnUpdateSessionCompleteDelegates,
-		[this, WeakOp = Op.AsWeak()](FName SessionName, const bool bWasSuccessful)
+		if (!bWasSuccessful)
 		{
-			if (TOnlineAsyncOpPtr<FUpdateSessionImpl> Op = WeakOp.Pin())
-			{
-				if (!bWasSuccessful)
-				{
-					Op->SetError(Errors::Unknown());
-					return;
-				}
-
-				const FUpdateSessionImpl::Params& OpParams = Op->GetParams();
+			Promise.EmplaceValue(Errors::Unknown());
+			return;
+		}
 
-				TOnlineResult<FGetMutableSessionByName> GetMutableSessionByNameResult = GetMutableSessionByName({ SessionName });
-				if (GetMutableSessionByNameResult.IsOk())
-				{
-					// We update our local session
-					TSharedRef<FSessionCommon>& FoundSession = GetMutableSessionByNameResult.GetOkValue().Session;
+		TOnlineResult<FGetMutableSessionByName> GetMutableSessionByNameResult = GetMutableSessionByName({ SessionName });
+		if (GetMutableSessionByNameResult.IsOk())
+		{
+			// We update our local session
+			TSharedRef<FSessionCommon>& FoundSession = GetMutableSessionByNameResult.GetOkValue().Session;
 
-					FSessionUpdate SessionUpdateData = BuildSessionUpdate(FoundSession, OpParams.Mutations);
+			FSessionUpdate SessionUpdateData = BuildSessionUpdate(FoundSession, OpParams.Mutations);
 
-					(*FoundSession) += SessionUpdateData;
+			(*FoundSession) += SessionUpdateData;
 
-					// We set the result and fire the event
-					Op->SetResult({ });
+			// We set the result and fire the event
+			Promise.EmplaceValue(FUpdateSessionImpl::Result{ });
 
-					FSessionUpdated SessionUpdatedEvent{ OpParams.SessionName, SessionUpdateData };
-					SessionEvents.OnSessionUpdated.Broadcast(SessionUpdatedEvent);
-				}
-				else
-				{
-					Op->SetError(Errors::NotFound());
-				}
-			}
-		});
+			FSessionUpdated SessionUpdatedEvent{ OpParams.SessionName, SessionUpdateData };
+			SessionEvents.OnSessionUpdated.Broadcast(SessionUpdatedEvent);
+		}
+		else
+		{
+			Promise.EmplaceValue(Errors::NotFound());
+		}
+	});
 
-		const TOnlineResult<FGetMutableSessionByName> GetMutableSessionByNameResult = GetMutableSessionByName({ OpParams.SessionName });
-		check(GetMutableSessionByNameResult.IsOk());
+	const TOnlineResult<FGetMutableSessionByName> GetMutableSessionByNameResult = GetMutableSessionByName({ Params.SessionName });
+	check(GetMutableSessionByNameResult.IsOk());
 		
-		const TSharedRef<FSessionCommon>& FoundSession = GetMutableSessionByNameResult.GetOkValue().Session;
+	const TSharedRef<FSessionCommon>& FoundSession = GetMutableSessionByNameResult.GetOkValue().Session;
 
-		// We will update a copy here, and wait until the operation has completed successfully to update our local data
-		FSessionSettings UpdatedV2Settings = FoundSession->SessionSettings;
-		//UpdatedV2Settings += OpParams.Mutations;
-		
-		// TODO: Will be refactored as part of the UpdateSession split
+	// We will update a copy here, and wait until the operation has completed successfully to update our local data
+	FSessionSettings UpdatedV2Settings = FoundSession->SessionSettings;
+	FSessionSettingsChanges SessionSettingsChanges = BuildSessionSettingsChanges(FoundSession, Params.Mutations.UpdatedSessionSettings);
+	UpdatedV2Settings += SessionSettingsChanges;
 
-		FOnlineSessionSettings UpdatedV1Settings = BuildV1SettingsForUpdate(OpParams.LocalAccountId, FoundSession);
+	FOnlineSessionSettings UpdatedV1Settings = BuildV1SettingsForUpdate(Params.LocalAccountId, FoundSession);
 
-		SessionsInterface->UpdateSession(OpParams.SessionName, UpdatedV1Settings);
-	})
-	.Enqueue(GetSerialQueue());
+	SessionsInterface->UpdateSession(Params.SessionName, UpdatedV1Settings);
 
-	return Op->GetHandle();
+	return Future;
 }
 
 TOnlineAsyncOpHandle<FLeaveSession> FSessionsOSSAdapter::LeaveSession(FLeaveSession::Params&& Params)
@@ -950,7 +938,7 @@ TOnlineAsyncOpHandle<FJoinSession> FSessionsOSSAdapter::JoinSession(FJoinSession
 
 				const FJoinSession::Params& OpParams = Op->GetParams();
 
-				TOnlineResult<FGetSessionById> GetSessionByIdResult = GetSessionById({ OpParams.LocalAccountId, OpParams.SessionId });
+				TOnlineResult<FGetSessionById> GetSessionByIdResult = GetSessionById({ OpParams.SessionId });
 				if (GetSessionByIdResult.IsError())
 				{
 					// If no result is found, the id might be expired, which we should notify
@@ -978,7 +966,7 @@ TOnlineAsyncOpHandle<FJoinSession> FSessionsOSSAdapter::JoinSession(FJoinSession
 			}
 		});
 
-		TOnlineResult<FGetSessionById> GetSessionByIdResult = GetSessionById({ OpParams.LocalAccountId, OpParams.SessionId });
+		TOnlineResult<FGetSessionById> GetSessionByIdResult = GetSessionById({ OpParams.SessionId });
 		if (GetSessionByIdResult.IsError())
 		{
 			Op.SetError(MoveTemp(GetSessionByIdResult.GetErrorValue()));
@@ -1153,7 +1141,7 @@ TOnlineResult<FGetResolvedConnectString> FSessionsOSSAdapter::GetResolvedConnect
 {
 	if (Params.SessionId.IsValid())
 	{
-		TOnlineResult<FGetSessionById> Result = GetSessionById({ Params.LocalAccountId, Params.SessionId });
+		TOnlineResult<FGetSessionById> Result = GetSessionById({ Params.SessionId });
 		if (Result.IsOk())
 		{
 			FOnlineSessionSearchResult SearchResult;
diff --git a/Engine/Plugins/Online/Experimental/OnlineServicesOSSAdapter/Source/Public/Online/SessionsOSSAdapter.h b/Engine/Plugins/Online/Experimental/OnlineServicesOSSAdapter/Source/Public/Online/SessionsOSSAdapter.h
--- a/Engine/Plugins/Online/Experimental/OnlineServicesOSSAdapter/Source/Public/Online/SessionsOSSAdapter.h
+++ b/Engine/Plugins/Online/Experimental/OnlineServicesOSSAdapter/Source/Public/Online/SessionsOSSAdapter.h
@@ -64,7 +64,6 @@ public:
 
 	// ISessions
 	virtual TOnlineAsyncOpHandle<FCreateSession> CreateSession(FCreateSession::Params&& Params) override;
-	virtual TOnlineAsyncOpHandle<FUpdateSessionImpl> UpdateSessionImpl(FUpdateSessionImpl::Params&& Params) override;
 	virtual TOnlineAsyncOpHandle<FLeaveSession> LeaveSession(FLeaveSession::Params&& Params) override;
 	virtual TOnlineAsyncOpHandle<FFindSessions> FindSessions(FFindSessions::Params&& Params) override;
 	virtual TOnlineAsyncOpHandle<FStartMatchmaking> StartMatchmaking(FStartMatchmaking::Params&& Params) override;
@@ -73,6 +72,9 @@ public:
 	virtual TOnlineAsyncOpHandle<FRemoveSessionMember> RemoveSessionMember(FRemoveSessionMember::Params&& Params) override;
 	virtual TOnlineAsyncOpHandle<FSendSessionInvite> SendSessionInvite(FSendSessionInvite::Params&& Params) override;
 
+	// FSessionsCommon
+	virtual TFuture<TOnlineResult<FUpdateSessionImpl>> UpdateSessionImpl(FUpdateSessionImpl::Params&& Params) override;
+
 	TOnlineResult<FGetResolvedConnectString> GetResolvedConnectString(const FGetResolvedConnectString::Params& Params);
 
 private:
