Removal of the concept of FSessionMember and MemberSettings from Sessions interface.
Refactored Sessions EOSGS to retrieve member information after session join with new method FUpdateSessionFromService
Refactor of FUpdateSessionImpl to return a TFuture for intended behaviour.
Miscellaneous fixes for errors discovered while testing changes.

diff --git a/Engine/Plugins/Online/Experimental/OnlineServices/Source/OnlineServicesCommon/Private/Online/NboSerializerCommonSvc.cpp b/Engine/Plugins/Online/Experimental/OnlineServices/Source/OnlineServicesCommon/Private/Online/NboSerializerCommonSvc.cpp
--- a/Engine/Plugins/Online/Experimental/OnlineServices/Source/OnlineServicesCommon/Private/Online/NboSerializerCommonSvc.cpp
+++ b/Engine/Plugins/Online/Experimental/OnlineServices/Source/OnlineServicesCommon/Private/Online/NboSerializerCommonSvc.cpp
@@ -58,11 +58,6 @@ void SerializeToBuffer(FNboSerializeToBuffer& Packet, const FCustomSessionSettin
 	}
 }
 
-void SerializeToBuffer(FNboSerializeToBuffer& Packet, const FSessionMember& SessionMember)
-{
-	SerializeToBuffer(Packet, SessionMember.MemberSettings);
-}
-
 void SerializeToBuffer(FNboSerializeToBuffer& Packet, const FSessionInfo& SessionInfo)
 {
 	Packet << SessionInfo.bAllowSanctionedPlayers;
@@ -157,11 +152,6 @@ void SerializeFromBuffer(FNboSerializeFromBuffer& Packet, FCustomSessionSettings
 	}
 }
 
-void SerializeFromBuffer(FNboSerializeFromBuffer& Packet, FSessionMember& SessionMember)
-{
-	SerializeFromBuffer(Packet, SessionMember.MemberSettings);
-}
-
 void SerializeFromBuffer(FNboSerializeFromBuffer& Packet, FSessionInfo& SessionInfo)
 {
 	Packet >> SessionInfo.bAllowSanctionedPlayers;
diff --git a/Engine/Plugins/Online/Experimental/OnlineServices/Source/OnlineServicesCommon/Private/Online/SessionsCommon.cpp b/Engine/Plugins/Online/Experimental/OnlineServices/Source/OnlineServicesCommon/Private/Online/SessionsCommon.cpp
--- a/Engine/Plugins/Online/Experimental/OnlineServices/Source/OnlineServicesCommon/Private/Online/SessionsCommon.cpp
+++ b/Engine/Plugins/Online/Experimental/OnlineServices/Source/OnlineServicesCommon/Private/Online/SessionsCommon.cpp
@@ -28,39 +28,6 @@ namespace UE::Online {
 
 		SessionMembers.Append(SessionUpdate.AddedSessionMembers);
 
-		for (const TPair<FAccountId, FSessionMemberChanges>& MemberEntry : SessionUpdate.SessionMembersChanges)
-		{
-			if (FSessionMember* SessionMember = SessionMembers.Find(MemberEntry.Key))
-			{
-				const FSessionMemberChanges& SessionMemberChanges = MemberEntry.Value;
-
-				for (const FName& Key : SessionMemberChanges.RemovedMemberSettings)
-				{
-					SessionMember->MemberSettings.Remove(Key);
-				}
-
-				SessionMember->MemberSettings.Append(SessionMemberChanges.AddedMemberSettings);
-
-				for (const TPair<FName, FCustomSessionSettingUpdate>& SettingEntry : SessionMemberChanges.ChangedMemberSettings)
-				{
-					if (FCustomSessionSetting* CustomSetting = SessionMember->MemberSettings.Find(SettingEntry.Key))
-					{
-						(*CustomSetting) = SettingEntry.Value.NewValue;
-					}
-				}
-			}
-		}
-
-		return *this;
-	}
-
-	/** FCombinedSessionUpdate */
-	FCombinedSessionUpdate& FCombinedSessionUpdate::operator+=(FCombinedSessionUpdate&& UpdatedValues)
-	{
-		UpdatedSessionSettings += MoveTemp(UpdatedValues.UpdatedSessionSettings);
-
-		UpdatedSessionMembers.Append(MoveTemp(UpdatedValues.UpdatedSessionMembers));
-
 		return *this;
 	}
 
@@ -88,7 +55,6 @@ namespace UE::Online {
 		RegisterCommand(&FSessionsCommon::ClearPresenceSession);
 		RegisterCommand(&FSessionsCommon::CreateSession);
 		RegisterCommand(&FSessionsCommon::UpdateSessionSettings);
-		RegisterCommand(&FSessionsCommon::UpdateSessionMember);
 		RegisterCommand(&FSessionsCommon::LeaveSession);
 		RegisterCommand(&FSessionsCommon::FindSessions);
 		RegisterCommand(&FSessionsCommon::StartMatchmaking);
@@ -96,6 +62,7 @@ namespace UE::Online {
 		RegisterCommand(&FSessionsCommon::AddSessionMember);
 		RegisterCommand(&FSessionsCommon::RemoveSessionMember);
 		RegisterCommand(&FSessionsCommon::SendSessionInvite);
+		RegisterCommand(&FSessionsCommon::GetSessionInvites);
 		RegisterCommand(&FSessionsCommon::RejectSessionInvite);
 	}
 
@@ -226,8 +193,8 @@ namespace UE::Online {
 				return;
 			}
 
-			UpdateSessionImpl({ OpParams.LocalAccountId, OpParams.SessionName, { OpParams.Mutations, {} } })
-			.Next([this, WeakOp = Op.AsWeak()](const TOnlineResult<FUpdateSessionImpl>& Result)
+			UpdateSessionSettingsImpl({ OpParams.LocalAccountId, OpParams.SessionName, OpParams.Mutations })
+			.Next([this, WeakOp = Op.AsWeak()](const TOnlineResult<FUpdateSessionSettingsImpl>& Result)
 			{
 				if (TOnlineAsyncOpPtr<FUpdateSessionSettings> StrongOp = WeakOp.Pin())
 				{
@@ -248,53 +215,6 @@ namespace UE::Online {
 		return Op->GetHandle();
 	}
 
-	TOnlineAsyncOpHandle<FUpdateSessionMember> FSessionsCommon::UpdateSessionMember(FUpdateSessionMember::Params&& Params)
-	{
-		// TODO: check if session members can update their own settings across all APIs or if they need the owner to do it
-
-		TOnlineAsyncOpRef<FUpdateSessionMember> Op = GetOp<FUpdateSessionMember>(MoveTemp(Params));
-		const FUpdateSessionMember::Params& OpParams = Op->GetParams();
-
-		FOnlineError ParamsCheck = CheckUpdateSessionMemberParams(OpParams);
-		if (ParamsCheck != Errors::Success())
-		{
-			Op->SetError(MoveTemp(ParamsCheck));
-			return Op->GetHandle();
-		}
-
-		Op->Then([this](TOnlineAsyncOp<FUpdateSessionMember>& Op) mutable
-		{
-			const FUpdateSessionMember::Params& OpParams = Op.GetParams();
-
-			FOnlineError StateCheck = CheckUpdateSessionMemberState(OpParams);
-			if (StateCheck != Errors::Success())
-			{
-				Op.SetError(MoveTemp(StateCheck));
-				return;
-			}
-
-			UpdateSessionImpl({ OpParams.LocalAccountId, OpParams.SessionName, { FSessionSettingsUpdate() , { { OpParams.LocalAccountId, OpParams.Mutations } } } })
-			.Next([this, WeakOp = Op.AsWeak()](const TOnlineResult<FUpdateSessionImpl>& Result)
-			{
-				if (TOnlineAsyncOpPtr<FUpdateSessionMember> StrongOp = WeakOp.Pin())
-				{
-					if (Result.IsOk())
-					{
-						StrongOp->SetResult({ });
-					}
-					else
-					{
-						FOnlineError ErrorValue = Result.GetErrorValue();
-						StrongOp->SetError(MoveTemp(ErrorValue));
-					}
-				}
-			});
-		})
-		.Enqueue(GetSerialQueue());
-
-		return Op->GetHandle();
-	}
-
 	TOnlineAsyncOpHandle<FLeaveSession> FSessionsCommon::LeaveSession(FLeaveSession::Params&& Params)
 	{
 		TOnlineAsyncOpRef<FLeaveSession> Operation = GetOp<FLeaveSession>(MoveTemp(Params));
@@ -595,19 +515,21 @@ namespace UE::Online {
 #define COPY_TOPTIONAL_VALUE_IF_SET(Value) \
 	if (UpdatedValues.Value.IsSet()) \
 	{ \
-		Result.Value = UpdatedValues.Value.GetValue(); \
+		SessionSettingsChanges.Value = UpdatedValues.Value.GetValue(); \
 	} \
 
-	FSessionSettingsChanges FSessionsCommon::BuildSessionSettingsChanges(const TSharedRef<FSessionCommon>& Session, const FSessionSettingsUpdate& UpdatedValues) const
+	FSessionUpdate FSessionsCommon::BuildSessionUpdate(const TSharedRef<FSessionCommon>& Session, const FSessionSettingsUpdate& UpdatedValues) const
 	{
-		FSessionSettingsChanges Result;
+		FSessionUpdate Result;
+
+		FSessionSettingsChanges SessionSettingsChanges;
 
 		COPY_TOPTIONAL_VALUE_IF_SET(SchemaName) // TODO: We may need some additional logic for schema changes
 		COPY_TOPTIONAL_VALUE_IF_SET(NumMaxConnections)
 		COPY_TOPTIONAL_VALUE_IF_SET(JoinPolicy)
 		COPY_TOPTIONAL_VALUE_IF_SET(bAllowNewMembers)
 
-		Result.RemovedCustomSettings.Append(UpdatedValues.RemovedCustomSettings);
+		SessionSettingsChanges.RemovedCustomSettings.Append(UpdatedValues.RemovedCustomSettings);
 
 		for (const TPair<FSchemaAttributeId, FCustomSessionSetting>& Entry : UpdatedValues.UpdatedCustomSettings)
 		{
@@ -615,57 +537,24 @@ namespace UE::Online {
 			{
 				FCustomSessionSettingUpdate SettingUpdate = { *CustomSetting, Entry.Value };
 
-				Result.ChangedCustomSettings.Emplace(Entry.Key, SettingUpdate);
+				SessionSettingsChanges.ChangedCustomSettings.Emplace(Entry.Key, SettingUpdate);
 			}
 			else
 			{
-				Result.AddedCustomSettings.Add(Entry);
+				SessionSettingsChanges.AddedCustomSettings.Add(Entry);
 			}
 		}
 
+		Result.SessionSettingsChanges = SessionSettingsChanges;
+
 		return Result;
 	}
 
 #undef COPY_TOPTIONAL_VALUE_IF_SET
 
-	FSessionUpdate FSessionsCommon::BuildSessionUpdate(const TSharedRef<FSessionCommon>& Session, const FCombinedSessionUpdate& UpdatedValues) const
+	TFuture<TOnlineResult<FUpdateSessionSettingsImpl>> FSessionsCommon::UpdateSessionSettingsImpl(FUpdateSessionSettingsImpl::Params&& Params)
 	{
-		FSessionUpdate Result;
-
-		Result.SessionSettingsChanges = BuildSessionSettingsChanges(Session, UpdatedValues.UpdatedSessionSettings);
-
-		for (const TPair<FAccountId, FSessionMemberUpdate>& SessionMemberEntry : UpdatedValues.UpdatedSessionMembers)
-		{
-			if (const FSessionMember* SessionMember = Session->GetSessionMembers().Find(SessionMemberEntry.Key))
-			{
-				FSessionMemberChanges SessionMemberChanges;
-
-				SessionMemberChanges.RemovedMemberSettings.Append(SessionMemberEntry.Value.RemovedMemberSettings);
-
-				for (const TPair <FSchemaAttributeId, FCustomSessionSetting>& SettingEntry : SessionMemberEntry.Value.UpdatedMemberSettings)
-				{
-					if (const FCustomSessionSetting* MemberSetting = SessionMember->MemberSettings.Find(SettingEntry.Key))
-					{
-						FCustomSessionSettingUpdate SettingUpdate = { *MemberSetting, SettingEntry.Value };
-
-						SessionMemberChanges.ChangedMemberSettings.Emplace(SettingEntry.Key, SettingUpdate);
-					}
-					else
-					{
-						SessionMemberChanges.AddedMemberSettings.Add(SettingEntry);
-					}
-				}
-
-				Result.SessionMembersChanges.Emplace(SessionMemberEntry.Key, SessionMemberChanges);
-			}
-		}
-
-		return Result;
-	}
-
-	TFuture<TOnlineResult<FUpdateSessionImpl>> FSessionsCommon::UpdateSessionImpl(FUpdateSessionImpl::Params&& Params)
-	{
-		return MakeFulfilledPromise<TOnlineResult<FUpdateSessionImpl>>(Errors::NotImplemented()).GetFuture();
+		return MakeFulfilledPromise<TOnlineResult<FUpdateSessionSettingsImpl>>(Errors::NotImplemented()).GetFuture();
 	}
 
 	TOnlineResult<FAddSessionMember> FSessionsCommon::AddSessionMemberImpl(const FAddSessionMember::Params& Params)
@@ -684,7 +573,7 @@ namespace UE::Online {
 			return TOnlineResult<FAddSessionMember>(Errors::InvalidState());
 		}
 
-		FoundSession->SessionMembers.Emplace(Params.LocalAccountId, Params.NewSessionMember);
+		FoundSession->SessionMembers.Emplace(Params.LocalAccountId);
 
 		return TOnlineResult<FAddSessionMember>();
 	}
@@ -817,51 +706,6 @@ namespace UE::Online {
 		return Errors::Success();
 	}
 
-	FOnlineError FSessionsCommon::CheckUpdateSessionMemberParams(const FUpdateSessionMember::Params& Params)
-	{
-		if (!Params.LocalAccountId.IsValid())
-		{
-			return Errors::InvalidUser();
-		}
-
-		if (Params.SessionName.IsNone())
-		{
-			return Errors::InvalidParams();
-		}
-
-		for (const TPair<FSchemaAttributeId, FCustomSessionSetting>& Entry : Params.Mutations.UpdatedMemberSettings)
-		{
-			if (Entry.Key.IsNone())
-			{
-				return Errors::InvalidParams();
-			}
-		}
-
-		return Errors::Success();
-	}
-
-	FOnlineError FSessionsCommon::CheckUpdateSessionMemberState(const FUpdateSessionMember::Params& Params)
-	{
-		// User login check
-		IAuthPtr Auth = Services.GetAuthInterface();
-		if (!Auth->IsLoggedIn(Params.LocalAccountId))
-		{
-			UE_LOG(LogTemp, Warning, TEXT("[FSessionsCommon::CheckUpdateSessionMemberState] User [%s] not logged in"), *ToLogString(Params.LocalAccountId));
-
-			return Errors::InvalidUser();
-		}
-
-		// Session name check
-		if (TOptional<FOnlineError> Result = CheckSessionExistsByName(Params.LocalAccountId, Params.SessionName))
-		{
-			UE_LOG(LogTemp, Warning, TEXT("[FSessionsCommon::CheckUpdateSessionMemberState] Session with name [%s] not found."), *Params.SessionName.ToString());
-
-			return Result.GetValue();
-		}
-
-		return Errors::Success();
-	}
-
 	FOnlineError FSessionsCommon::CheckFindSessionsParams(const FFindSessions::Params& Params)
 	{
 		if (Params.MaxResults <= 0)
diff --git a/Engine/Plugins/Online/Experimental/OnlineServices/Source/OnlineServicesCommon/Public/Online/NboSerializerCommonSvc.h b/Engine/Plugins/Online/Experimental/OnlineServices/Source/OnlineServicesCommon/Public/Online/NboSerializerCommonSvc.h
--- a/Engine/Plugins/Online/Experimental/OnlineServices/Source/OnlineServicesCommon/Public/Online/NboSerializerCommonSvc.h
+++ b/Engine/Plugins/Online/Experimental/OnlineServices/Source/OnlineServicesCommon/Public/Online/NboSerializerCommonSvc.h
@@ -18,7 +18,6 @@ namespace NboSerializerCommonSvc {
 ONLINESERVICESCOMMON_API void SerializeToBuffer(FNboSerializeToBuffer& Packet, const FSchemaVariant& Data);
 ONLINESERVICESCOMMON_API void SerializeToBuffer(FNboSerializeToBuffer& Packet, const FCustomSessionSetting& CustomSessionSetting);
 ONLINESERVICESCOMMON_API void SerializeToBuffer(FNboSerializeToBuffer& Packet, const FCustomSessionSettingsMap& CustomSessionSettingsMap);
-ONLINESERVICESCOMMON_API void SerializeToBuffer(FNboSerializeToBuffer& Packet, const FSessionMember& SessionMember);
 ONLINESERVICESCOMMON_API void SerializeToBuffer(FNboSerializeToBuffer& Packet, const FSessionSettings& SessionSettings);
 ONLINESERVICESCOMMON_API void SerializeToBuffer(FNboSerializeToBuffer& Packet, const FSessionCommon& Session);
 
@@ -27,7 +26,6 @@ ONLINESERVICESCOMMON_API void SerializeToBuffer(FNboSerializeToBuffer& Packet, c
 ONLINESERVICESCOMMON_API void SerializeFromBuffer(FNboSerializeFromBuffer& Packet, FSchemaVariant& Data);
 ONLINESERVICESCOMMON_API void SerializeFromBuffer(FNboSerializeFromBuffer& Packet, FCustomSessionSetting& CustomSessionSetting);
 ONLINESERVICESCOMMON_API void SerializeFromBuffer(FNboSerializeFromBuffer& Packet, FCustomSessionSettingsMap& CustomSessionSettingsMap);
-ONLINESERVICESCOMMON_API void SerializeFromBuffer(FNboSerializeFromBuffer& Packet, FSessionMember& SessionMember);
 ONLINESERVICESCOMMON_API void SerializeFromBuffer(FNboSerializeFromBuffer& Packet, FSessionSettings& SessionSettings);
 ONLINESERVICESCOMMON_API void SerializeFromBuffer(FNboSerializeFromBuffer& Packet, FSessionCommon& Session);
 
diff --git a/Engine/Plugins/Online/Experimental/OnlineServices/Source/OnlineServicesCommon/Public/Online/SessionsCommon.h b/Engine/Plugins/Online/Experimental/OnlineServices/Source/OnlineServicesCommon/Public/Online/SessionsCommon.h
--- a/Engine/Plugins/Online/Experimental/OnlineServices/Source/OnlineServicesCommon/Public/Online/SessionsCommon.h
+++ b/Engine/Plugins/Online/Experimental/OnlineServices/Source/OnlineServicesCommon/Public/Online/SessionsCommon.h
@@ -140,7 +140,7 @@ public:
 	virtual const uint32 GetNumOpenConnections() const override				{ return SessionSettings.NumMaxConnections - SessionMembers.Num(); }
 	virtual const FSessionInfo& GetSessionInfo() const override				{ return SessionInfo; }
 	virtual const FSessionSettings GetSessionSettings() const override		{ return SessionSettings; }
-	virtual const FSessionMembersMap& GetSessionMembers() const override	{ return SessionMembers; }
+	virtual const FSessionMembersSet& GetSessionMembers() const override	{ return SessionMembers; }
 
 	virtual bool IsJoinable() const override								{ return GetNumOpenConnections() > 0 && SessionSettings.bAllowNewMembers; }
 
@@ -158,8 +158,8 @@ public:
 	/** Set of session properties that can be altered by the session owner */
 	FSessionSettings SessionSettings;
 
-	/* Map of session member ids to their corresponding user-defined properties */
-	FSessionMembersMap SessionMembers;
+	/* Set containing user ids for all the session members */
+	FSessionMembersSet SessionMembers;
 
 	FSessionCommon& operator+=(const FSessionUpdate& SessionUpdate);
 };
@@ -198,20 +198,9 @@ struct FGetMutableSessionById
 	};
 };
 
-struct FCombinedSessionUpdate
+struct FUpdateSessionSettingsImpl
 {
-	/** Changes to current session settings */
-	FSessionSettingsUpdate UpdatedSessionSettings;
-
-	/** Map of changes to current member settings */
-	FSessionMemberUpdatesMap UpdatedSessionMembers;
-
-	FCombinedSessionUpdate& operator+=(FCombinedSessionUpdate&& UpdatedValue);
-};
-
-struct FUpdateSessionImpl
-{
-	static constexpr TCHAR Name[] = TEXT("UpdateSessionImpl");
+	static constexpr TCHAR Name[] = TEXT("UpdateSessionSettingsImpl");
 
 	struct Params
 	{
@@ -221,8 +210,8 @@ struct FUpdateSessionImpl
 		/** The local name for the session */
 		FName SessionName;
 
-		/** Combined data for a session update */
-		FCombinedSessionUpdate Mutations;
+		/** Changes to current session settings */
+		FSessionSettingsUpdate Mutations;
 	};
 
 	struct Result
@@ -252,7 +241,6 @@ public:
 	virtual TOnlineResult<FClearPresenceSession> ClearPresenceSession(FClearPresenceSession::Params&& Params) override;
 	virtual TOnlineAsyncOpHandle<FCreateSession> CreateSession(FCreateSession::Params&& Params) override;
 	virtual TOnlineAsyncOpHandle<FUpdateSessionSettings> UpdateSessionSettings(FUpdateSessionSettings::Params&& Params) override;
-	virtual TOnlineAsyncOpHandle<FUpdateSessionMember> UpdateSessionMember(FUpdateSessionMember::Params&& Params) override;
 	virtual TOnlineAsyncOpHandle<FLeaveSession> LeaveSession(FLeaveSession::Params&& Params) override;
 	virtual TOnlineAsyncOpHandle<FFindSessions> FindSessions(FFindSessions::Params&& Params) override;
 	virtual TOnlineAsyncOpHandle<FStartMatchmaking> StartMatchmaking(FStartMatchmaking::Params&& Params) override;
@@ -281,10 +269,9 @@ protected:
 	void ClearSessionInvitesForSession(const FAccountId& LocalAccountId, const FOnlineSessionId SessionId);
 	void ClearSessionReferences(const FOnlineSessionId SessionId, const FName& SessionName, const FAccountId& LocalAccountId);
 
-	FSessionSettingsChanges BuildSessionSettingsChanges(const TSharedRef<FSessionCommon>& Session, const FSessionSettingsUpdate& UpdatedValues) const;
-	FSessionUpdate BuildSessionUpdate(const TSharedRef<FSessionCommon>& Session, const FCombinedSessionUpdate& UpdatedValues) const;
+	FSessionUpdate BuildSessionUpdate(const TSharedRef<FSessionCommon>& Session, const FSessionSettingsUpdate& UpdatedValues) const;
 
-	virtual TFuture<TOnlineResult<FUpdateSessionImpl>> UpdateSessionImpl(FUpdateSessionImpl::Params&& Params);
+	virtual TFuture<TOnlineResult<FUpdateSessionSettingsImpl>> UpdateSessionSettingsImpl(FUpdateSessionSettingsImpl::Params&& Params);
 	TOnlineResult<FAddSessionMember> AddSessionMemberImpl(const FAddSessionMember::Params& Params);
 	TOnlineResult<FRemoveSessionMember> RemoveSessionMemberImpl(const FRemoveSessionMember::Params& Params);
 
@@ -294,9 +281,6 @@ protected:
 	FOnlineError CheckUpdateSessionSettingsParams(const FUpdateSessionSettings::Params& Params);
 	FOnlineError CheckUpdateSessionSettingsState(const FUpdateSessionSettings::Params& Params);
 
-	FOnlineError CheckUpdateSessionMemberParams(const FUpdateSessionMember::Params& Params);
-	FOnlineError CheckUpdateSessionMemberState(const FUpdateSessionMember::Params& Params);
-
 	FOnlineError CheckFindSessionsParams(const FFindSessions::Params& Params);
 	FOnlineError CheckFindSessionsState(const FFindSessions::Params& Params);
 
@@ -373,18 +357,13 @@ BEGIN_ONLINE_STRUCT_META(FGetMutableSessionById::Result)
 	ONLINE_STRUCT_FIELD(FGetMutableSessionById::Result, Session)
 END_ONLINE_STRUCT_META()
 
-BEGIN_ONLINE_STRUCT_META(FCombinedSessionUpdate)
-	ONLINE_STRUCT_FIELD(FCombinedSessionUpdate, UpdatedSessionSettings),
-	ONLINE_STRUCT_FIELD(FCombinedSessionUpdate, UpdatedSessionMembers)
-END_ONLINE_STRUCT_META()
-
-BEGIN_ONLINE_STRUCT_META(FUpdateSessionImpl::Params)
-	ONLINE_STRUCT_FIELD(FUpdateSessionImpl::Params, LocalAccountId),
-	ONLINE_STRUCT_FIELD(FUpdateSessionImpl::Params, SessionName),
-	ONLINE_STRUCT_FIELD(FUpdateSessionImpl::Params, Mutations)
+BEGIN_ONLINE_STRUCT_META(FUpdateSessionSettingsImpl::Params)
+	ONLINE_STRUCT_FIELD(FUpdateSessionSettingsImpl::Params, LocalAccountId),
+	ONLINE_STRUCT_FIELD(FUpdateSessionSettingsImpl::Params, SessionName),
+	ONLINE_STRUCT_FIELD(FUpdateSessionSettingsImpl::Params, Mutations)
 END_ONLINE_STRUCT_META()
 
-BEGIN_ONLINE_STRUCT_META(FUpdateSessionImpl::Result)
+BEGIN_ONLINE_STRUCT_META(FUpdateSessionSettingsImpl::Result)
 END_ONLINE_STRUCT_META()
 
 /* Meta*/ }
diff --git a/Engine/Plugins/Online/Experimental/OnlineServices/Source/OnlineServicesCommonEngineUtils/Private/Online/SessionsLAN.cpp b/Engine/Plugins/Online/Experimental/OnlineServices/Source/OnlineServicesCommonEngineUtils/Private/Online/SessionsLAN.cpp
--- a/Engine/Plugins/Online/Experimental/OnlineServices/Source/OnlineServicesCommonEngineUtils/Private/Online/SessionsLAN.cpp
+++ b/Engine/Plugins/Online/Experimental/OnlineServices/Source/OnlineServicesCommonEngineUtils/Private/Online/SessionsLAN.cpp
@@ -162,7 +162,7 @@ TOnlineAsyncOpHandle<FCreateSession> FSessionsLAN::CreateSession(FCreateSession:
 		NewSessionLANRef->SessionSettings = OpParams.SessionSettings;
 
 		// For LAN sessions, we'll add the Session member manually instead of calling JoinSession since there is no API calls involved
-		NewSessionLANRef->SessionMembers.Emplace(OpParams.LocalAccountId, OpParams.SessionMemberData);
+		NewSessionLANRef->SessionMembers.Emplace(OpParams.LocalAccountId);
 
 		// We save the local object for the session, and set up the appropriate references
 		AddSessionWithReferences(NewSessionLANRef, OpParams.SessionName, OpParams.LocalAccountId, OpParams.bPresenceEnabled);
@@ -174,10 +174,10 @@ TOnlineAsyncOpHandle<FCreateSession> FSessionsLAN::CreateSession(FCreateSession:
 	return Op->GetHandle();
 }
 
-TFuture<TOnlineResult<FUpdateSessionImpl>> FSessionsLAN::UpdateSessionImpl(FUpdateSessionImpl::Params&& Params)
+TFuture<TOnlineResult<FUpdateSessionSettingsImpl>> FSessionsLAN::UpdateSessionSettingsImpl(FUpdateSessionSettingsImpl::Params&& Params)
 {
-	TPromise<TOnlineResult<FUpdateSessionImpl>> Promise;
-	TFuture<TOnlineResult<FUpdateSessionImpl>> Future = Promise.GetFuture();
+	TPromise<TOnlineResult<FUpdateSessionSettingsImpl>> Promise;
+	TFuture<TOnlineResult<FUpdateSessionSettingsImpl>> Future = Promise.GetFuture();
 
 	TOnlineResult<FGetMutableSessionByName> GetMutableSessionByNameResult = GetMutableSessionByName({ Params.SessionName });
 	check(GetMutableSessionByNameResult.IsOk());
@@ -186,13 +186,12 @@ TFuture<TOnlineResult<FUpdateSessionImpl>> FSessionsLAN::UpdateSessionImpl(FUpda
 
 	TSharedRef<FSessionCommon> FoundSession = GetMutableSessionByNameResult.GetOkValue().Session;
 	FSessionSettings& SessionSettings = FoundSession->SessionSettings;
-	const FCombinedSessionUpdate& SessionUpdate = Params.Mutations;
 
-	if (SessionUpdate.UpdatedSessionSettings.JoinPolicy.IsSet())
+	if (Params.Mutations.JoinPolicy.IsSet())
 	{
 		// Changes in the join policy setting will mean we start or stop the LAN Beacon broadcasting the session information
 		// There is no FriendsLAN interface at the time of this implementation, hence the binary behavior (Public/Non-Public)
-		if (SessionSettings.JoinPolicy != ESessionJoinPolicy::Public && SessionUpdate.UpdatedSessionSettings.JoinPolicy.GetValue() == ESessionJoinPolicy::Public)
+		if (SessionSettings.JoinPolicy != ESessionJoinPolicy::Public && Params.Mutations.JoinPolicy.GetValue() == ESessionJoinPolicy::Public)
 		{
 			if (TOptional<FOnlineError> TryHostLANSessionResult = TryHostLANSession())
 			{
@@ -200,7 +199,7 @@ TFuture<TOnlineResult<FUpdateSessionImpl>> FSessionsLAN::UpdateSessionImpl(FUpda
 				return Future;
 			}
 		}
-		else if (SessionSettings.JoinPolicy == ESessionJoinPolicy::Public && SessionUpdate.UpdatedSessionSettings.JoinPolicy.GetValue() != ESessionJoinPolicy::Public)
+		else if (SessionSettings.JoinPolicy == ESessionJoinPolicy::Public && Params.Mutations.JoinPolicy.GetValue() != ESessionJoinPolicy::Public)
 		{
 			StopLANSession();
 		}
@@ -212,7 +211,7 @@ TFuture<TOnlineResult<FUpdateSessionImpl>> FSessionsLAN::UpdateSessionImpl(FUpda
 	(*FoundSession) += SessionUpdateData;
 
 	// We set the result and fire the event
-	Promise.EmplaceValue(FUpdateSessionImpl::Result{ });
+	Promise.EmplaceValue(FUpdateSessionSettingsImpl::Result{ });
 
 	FSessionUpdated SessionUpdatedEvent{ Params.SessionName, SessionUpdateData };
 	SessionEvents.OnSessionUpdated.Broadcast(SessionUpdatedEvent);
diff --git a/Engine/Plugins/Online/Experimental/OnlineServices/Source/OnlineServicesCommonEngineUtils/Public/Online/SessionsLAN.h b/Engine/Plugins/Online/Experimental/OnlineServices/Source/OnlineServicesCommonEngineUtils/Public/Online/SessionsLAN.h
--- a/Engine/Plugins/Online/Experimental/OnlineServices/Source/OnlineServicesCommonEngineUtils/Public/Online/SessionsLAN.h
+++ b/Engine/Plugins/Online/Experimental/OnlineServices/Source/OnlineServicesCommonEngineUtils/Public/Online/SessionsLAN.h
@@ -62,7 +62,7 @@ public:
 	virtual TOnlineAsyncOpHandle<FLeaveSession> LeaveSession(FLeaveSession::Params&& Params) override;
 
 	// FSessionsCommon
-	virtual TFuture<TOnlineResult<FUpdateSessionImpl>> UpdateSessionImpl(FUpdateSessionImpl::Params&& Params) override;
+	virtual TFuture<TOnlineResult<FUpdateSessionSettingsImpl>> UpdateSessionSettingsImpl(FUpdateSessionSettingsImpl::Params&& Params) override;
 
 protected:
 	TOptional<FOnlineError> TryHostLANSession();
diff --git a/Engine/Plugins/Online/Experimental/OnlineServices/Source/OnlineServicesInterface/Private/Online/Sessions.cpp b/Engine/Plugins/Online/Experimental/OnlineServices/Source/OnlineServicesInterface/Private/Online/Sessions.cpp
--- a/Engine/Plugins/Online/Experimental/OnlineServices/Source/OnlineServicesInterface/Private/Online/Sessions.cpp
+++ b/Engine/Plugins/Online/Experimental/OnlineServices/Source/OnlineServicesInterface/Private/Online/Sessions.cpp
@@ -36,45 +36,6 @@ void LexFromString(ESessionJoinPolicy& Value, const TCHAR* InStr)
 	}
 }
 
-FSessionMemberUpdate& FSessionMemberUpdate::operator+=(FSessionMemberUpdate&& UpdatedValue)
-{
-	for (const TPair<FSchemaAttributeId, FCustomSessionSetting>& UpdatedMemberSetting : UpdatedValue.UpdatedMemberSettings)
-	{
-		// If an update adds a modification to a setting that had previously been marked for removal, we'll keep the latest change
-		RemovedMemberSettings.Remove(UpdatedMemberSetting.Key);
-	}
-	UpdatedMemberSettings.Append(MoveTemp(UpdatedValue.UpdatedMemberSettings));
-
-	for (FSchemaAttributeId& Key : UpdatedValue.RemovedMemberSettings)
-	{
-		// If an update removes a setting that had previously been modified, we'll keep the latest change
-		UpdatedMemberSettings.Remove(Key);
-		RemovedMemberSettings.AddUnique(MoveTemp(Key));
-	}
-
-	return *this;
-}
-
-FSessionMember& FSessionMember::operator+=(const FSessionMemberChanges& MemberChanges)
-{
-	for (const FName& Key : MemberChanges.RemovedMemberSettings)
-	{
-		MemberSettings.Remove(Key);
-	}
-
-	MemberSettings.Append(MemberChanges.AddedMemberSettings);
-
-	for (const TPair<FName, FCustomSessionSettingUpdate>& SettingEntry : MemberChanges.ChangedMemberSettings)
-	{
-		if (FCustomSessionSetting* CustomSetting = MemberSettings.Find(SettingEntry.Key))
-		{
-			(*CustomSetting) = SettingEntry.Value.NewValue;
-		}
-	}
-
-	return *this;
-}
-
 #define MOVE_TOPTIONAL_IF_SET(Value) \
 	if (UpdatedValue.Value.IsSet()) \
 	{ \
diff --git a/Engine/Plugins/Online/Experimental/OnlineServices/Source/OnlineServicesInterface/Public/Online/Sessions.h b/Engine/Plugins/Online/Experimental/OnlineServices/Source/OnlineServicesInterface/Public/Online/Sessions.h
--- a/Engine/Plugins/Online/Experimental/OnlineServices/Source/OnlineServicesInterface/Public/Online/Sessions.h
+++ b/Engine/Plugins/Online/Experimental/OnlineServices/Source/OnlineServicesInterface/Public/Online/Sessions.h
@@ -47,44 +47,7 @@ struct FCustomSessionSettingUpdate
 
 using FCustomSessionSettingUpdateMap = TMap<FName, FCustomSessionSettingUpdate>;
 
-// Session Member Settings
-
-/** Contains new values for a SessionMember's attributes. Taken as a parameter by the FUpdateSessionMember method */
-struct ONLINESERVICESINTERFACE_API FSessionMemberUpdate
-{
-	FCustomSessionSettingsMap UpdatedMemberSettings;
-	TArray<FSchemaAttributeId> RemovedMemberSettings;
-
-	FSessionMemberUpdate& operator+=(FSessionMemberUpdate&& UpdatedValue);
-};
-
-using FSessionMemberUpdatesMap = TMap<FAccountId, FSessionMemberUpdate>;
-
-/** Contains updated data for any modifiable attributes of FSessionMember. Member of FSessionUpdated event */
-struct FSessionMemberChanges
-{
-	/** New custom settings, with their values */
-	FCustomSessionSettingsMap AddedMemberSettings;
-
-	/** Existing custom settings that changed value, including new and old values */
-	FCustomSessionSettingUpdateMap ChangedMemberSettings;
-
-	/** Keys for removed custom settings */
-	TArray<FName> RemovedMemberSettings;
-};
-
-using FSessionMemberChangesMap = TMap<FAccountId, FSessionMemberChanges>;
-
-/** A member is a player that is part of the session, and they stop being a member when they leave it */
-struct ONLINESERVICESINTERFACE_API FSessionMember
-{
-	/* Map of user-defined settings to be passed to the platform APIs as additional information for the session member */
-	FCustomSessionSettingsMap MemberSettings;
-
-	FSessionMember& operator+=(const FSessionMemberChanges& Changes);
-};
-
-using FSessionMembersMap = TMap<FAccountId, FSessionMember>;
+using FSessionMembersSet = TSet<FAccountId>;
 
 enum class ESessionJoinPolicy : uint8
 {
@@ -188,7 +151,7 @@ public:
 	virtual const uint32 GetNumOpenConnections() const = 0;
 	virtual const FSessionInfo& GetSessionInfo() const = 0;
 	virtual const FSessionSettings GetSessionSettings() const = 0;
-	virtual const FSessionMembersMap& GetSessionMembers() const = 0;
+	virtual const FSessionMembersSet& GetSessionMembers() const = 0;
 
 	/** Evaluates a series of factors to determine if a session is accepting new members */
 	virtual bool IsJoinable() const = 0;
@@ -337,9 +300,6 @@ struct FCreateSession
 		/** The local name for the session */
 		FName SessionName;
 
-		/** Information for the local user who will join the session after creation */
-		FSessionMember SessionMemberData;
-
 		/* In platforms that support this feature, it will set the session id to this value. Might be subject to minimum and maximum length */
 		FString SessionIdOverride;
 
@@ -390,28 +350,6 @@ struct FUpdateSessionSettings
 	};
 };
 
-struct FUpdateSessionMember
-{
-	static constexpr TCHAR Name[] = TEXT("UpdateSessionMember");
-
-	struct Params
-	{
-		/** The local user agent which will perform the action. */
-		FAccountId LocalAccountId;
-
-		/** The local name for the session */
-		FName SessionName;
-
-		/** Changes to current session settings */
-		FSessionMemberUpdate Mutations;
-	};
-
-	struct Result
-	{
-
-	};
-};
-
 struct FLeaveSession
 {
 	static constexpr TCHAR Name[] = TEXT("LeaveSession");
@@ -499,9 +437,6 @@ struct FJoinSession
 		/* Id handle for the session to be joined. To be retrieved via session search or invite */
 		FOnlineSessionId SessionId;
 
-		/* Information for the local user who will join the session */
-		FSessionMember SessionMemberData;
-
 		/* Whether this session should be set as the user's new presence session. False by default */
 		bool bPresenceEnabled = false;
 	};
@@ -523,9 +458,6 @@ struct FAddSessionMember
 
 		/* Local name for the session */
 		FName SessionName;
-
-		/** Information for the session member to be added to the session. Any player that joins the session becomes a new member in doing so */
-		FSessionMember NewSessionMember;
 	};
 
 	struct Result
@@ -637,13 +569,10 @@ struct ONLINESERVICESINTERFACE_API FSessionUpdate
 	TOptional<FSessionSettingsChanges> SessionSettingsChanges;
 
 	/** Session member information for members that just joined the session */
-	FSessionMembersMap AddedSessionMembers;
-
-	/** Updated values for session member information */
-	FSessionMemberChangesMap SessionMembersChanges;
+	FSessionMembersSet AddedSessionMembers;
 
 	/** Id handles for members that just left the session */
-	TArray<FAccountId> RemovedSessionMembers;
+	FSessionMembersSet RemovedSessionMembers;
 
 	FSessionUpdate& operator+=(const FSessionUpdate& SessionUpdate);
 };
@@ -768,14 +697,6 @@ public:
 	 */
 	virtual TOnlineAsyncOpHandle<FUpdateSessionSettings> UpdateSessionSettings(FUpdateSessionSettings::Params&& Params) = 0;
 
-	/**
-	 * Update a given session member's information. Can be called by any session member to change only their own information.
-	 *
-	 * @param Parameters for the UpdateSessionMember call
-	 * @return
-	 */
-	virtual TOnlineAsyncOpHandle<FUpdateSessionMember> UpdateSessionMember(FUpdateSessionMember::Params&& Params) = 0;
-
 	/**
 	 * Leave and optionally destroy a given session.
 	 *
@@ -914,15 +835,6 @@ BEGIN_ONLINE_STRUCT_META(FCustomSessionSettingUpdate)
 	ONLINE_STRUCT_FIELD(FCustomSessionSettingUpdate, NewValue)
 END_ONLINE_STRUCT_META()
 
-BEGIN_ONLINE_STRUCT_META(FSessionMember)
-	ONLINE_STRUCT_FIELD(FSessionMember, MemberSettings)
-END_ONLINE_STRUCT_META()
-
-BEGIN_ONLINE_STRUCT_META(FSessionMemberUpdate)
-	ONLINE_STRUCT_FIELD(FSessionMemberUpdate, UpdatedMemberSettings),
-	ONLINE_STRUCT_FIELD(FSessionMemberUpdate, RemovedMemberSettings)
-END_ONLINE_STRUCT_META()
-
 BEGIN_ONLINE_STRUCT_META(FSessionSettings)
 	ONLINE_STRUCT_FIELD(FSessionSettings, SchemaName),
 	ONLINE_STRUCT_FIELD(FSessionSettings, NumMaxConnections),
@@ -1015,7 +927,6 @@ END_ONLINE_STRUCT_META()
 BEGIN_ONLINE_STRUCT_META(FCreateSession::Params)
 	ONLINE_STRUCT_FIELD(FCreateSession::Params, LocalAccountId),
 	ONLINE_STRUCT_FIELD(FCreateSession::Params, SessionName),
-	ONLINE_STRUCT_FIELD(FCreateSession::Params, SessionMemberData),
 	ONLINE_STRUCT_FIELD(FCreateSession::Params, SessionIdOverride),
 	ONLINE_STRUCT_FIELD(FCreateSession::Params, bPresenceEnabled),
 	ONLINE_STRUCT_FIELD(FCreateSession::Params, bIsLANSession),
@@ -1037,15 +948,6 @@ END_ONLINE_STRUCT_META()
 BEGIN_ONLINE_STRUCT_META(FUpdateSessionSettings::Result)
 END_ONLINE_STRUCT_META()
 
-BEGIN_ONLINE_STRUCT_META(FUpdateSessionMember::Params)
-	ONLINE_STRUCT_FIELD(FUpdateSessionMember::Params, LocalAccountId),
-	ONLINE_STRUCT_FIELD(FUpdateSessionMember::Params, SessionName),
-	ONLINE_STRUCT_FIELD(FUpdateSessionMember::Params, Mutations)
-END_ONLINE_STRUCT_META()
-
-BEGIN_ONLINE_STRUCT_META(FUpdateSessionMember::Result)
-END_ONLINE_STRUCT_META()
-
 BEGIN_ONLINE_STRUCT_META(FLeaveSession::Params)
 	ONLINE_STRUCT_FIELD(FLeaveSession::Params, LocalAccountId),
 	ONLINE_STRUCT_FIELD(FLeaveSession::Params, SessionName),
@@ -1080,7 +982,6 @@ BEGIN_ONLINE_STRUCT_META(FJoinSession::Params)
 	ONLINE_STRUCT_FIELD(FJoinSession::Params, LocalAccountId),
 	ONLINE_STRUCT_FIELD(FJoinSession::Params, SessionName),
 	ONLINE_STRUCT_FIELD(FJoinSession::Params, SessionId),
-	ONLINE_STRUCT_FIELD(FJoinSession::Params, SessionMemberData),
 	ONLINE_STRUCT_FIELD(FJoinSession::Params, bPresenceEnabled)
 END_ONLINE_STRUCT_META()
 
@@ -1089,8 +990,7 @@ END_ONLINE_STRUCT_META()
 
 BEGIN_ONLINE_STRUCT_META(FAddSessionMember::Params)
 	ONLINE_STRUCT_FIELD(FAddSessionMember::Params, LocalAccountId),
-	ONLINE_STRUCT_FIELD(FAddSessionMember::Params, SessionName),
-	ONLINE_STRUCT_FIELD(FAddSessionMember::Params, NewSessionMember)
+	ONLINE_STRUCT_FIELD(FAddSessionMember::Params, SessionName)
 END_ONLINE_STRUCT_META()
 
 BEGIN_ONLINE_STRUCT_META(FAddSessionMember::Result)
@@ -1139,16 +1039,10 @@ BEGIN_ONLINE_STRUCT_META(FSessionSettingsChanges)
 	ONLINE_STRUCT_FIELD(FSessionSettingsChanges, RemovedCustomSettings)
 END_ONLINE_STRUCT_META()
 
-BEGIN_ONLINE_STRUCT_META(FSessionMemberChanges)
-	ONLINE_STRUCT_FIELD(FSessionMemberChanges, AddedMemberSettings),
-	ONLINE_STRUCT_FIELD(FSessionMemberChanges, ChangedMemberSettings),
-	ONLINE_STRUCT_FIELD(FSessionMemberChanges, RemovedMemberSettings)
-END_ONLINE_STRUCT_META()
-
 BEGIN_ONLINE_STRUCT_META(FSessionUpdate)
 	ONLINE_STRUCT_FIELD(FSessionUpdate, OwnerAccountId),
 	ONLINE_STRUCT_FIELD(FSessionUpdate, SessionSettingsChanges),
-	ONLINE_STRUCT_FIELD(FSessionUpdate, SessionMembersChanges),
+	ONLINE_STRUCT_FIELD(FSessionUpdate, AddedSessionMembers),
 	ONLINE_STRUCT_FIELD(FSessionUpdate, RemovedSessionMembers)
 END_ONLINE_STRUCT_META()
 
diff --git a/Engine/Plugins/Online/Experimental/OnlineServicesEOSGS/Source/Private/Online/SessionsEOSGS.cpp b/Engine/Plugins/Online/Experimental/OnlineServicesEOSGS/Source/Private/Online/SessionsEOSGS.cpp
--- a/Engine/Plugins/Online/Experimental/OnlineServicesEOSGS/Source/Private/Online/SessionsEOSGS.cpp
+++ b/Engine/Plugins/Online/Experimental/OnlineServicesEOSGS/Source/Private/Online/SessionsEOSGS.cpp
@@ -42,15 +42,15 @@ FOnlineSessionInviteIdRegistryEOSGS& FOnlineSessionInviteIdRegistryEOSGS::Get()
 
 /** FSessionEOSGS */
 
-FSessionEOSGS::FSessionEOSGS(const EOS_HSessionDetails& InSessionDetailsHandle)
-	: SessionDetailsHandle(MakeShared<FSessionDetailsHandleEOSGS>(InSessionDetailsHandle))
+FSessionEOSGS::FSessionEOSGS(const TSharedPtr<FSessionDetailsHandleEOSGS>& InSessionDetailsHandle)
+	: SessionDetailsHandle(InSessionDetailsHandle)
 {
 	EOS_SessionDetails_CopyInfoOptions CopyInfoOptions = { };
 	CopyInfoOptions.ApiVersion = EOS_SESSIONDETAILS_COPYINFO_API_LATEST;
 	static_assert(EOS_SESSIONDETAILS_COPYINFO_API_LATEST == 1, "EOS_SessionDetails_CopyInfoOptions updated, check new fields");
 
 	EOS_SessionDetails_Info* SessionDetailsInfo = nullptr;
-	EOS_EResult CopyInfoResult = EOS_SessionDetails_CopyInfo(InSessionDetailsHandle, &CopyInfoOptions, &SessionDetailsInfo);
+	EOS_EResult CopyInfoResult = EOS_SessionDetails_CopyInfo(SessionDetailsHandle->SessionDetailsHandle, &CopyInfoOptions, &SessionDetailsInfo);
 	if (CopyInfoResult == EOS_EResult::EOS_Success)
 	{
 		SessionInfo.SessionId = FSessionsEOSGS::CreateSessionId(FString(SessionDetailsInfo->SessionId));
@@ -62,16 +62,18 @@ FSessionEOSGS::FSessionEOSGS(const EOS_HSessionDetails& InSessionDetailsHandle)
 		SessionSettings.bAllowNewMembers = SessionDetailsInfo->Settings->bInvitesAllowed && SessionDetailsInfo->Settings->bAllowJoinInProgress;
 		SessionSettings.NumMaxConnections = SessionDetailsInfo->Settings->NumPublicConnections;
 		SessionSettings.JoinPolicy = FromServiceType(SessionDetailsInfo->Settings->PermissionLevel);
-		SessionSettings.CustomSettings.Emplace(EOSGS_BUCKET_ID, FCustomSessionSetting{ FSchemaVariant(FString(UTF8_TO_TCHAR(SessionDetailsInfo->Settings->BucketId))), ESchemaAttributeVisibility::Public });
+		SessionSettings.CustomSettings.Emplace(EOSGS_BUCKET_ID_ATTRIBUTE_KEY, FCustomSessionSetting{ FSchemaVariant(FString(UTF8_TO_TCHAR(SessionDetailsInfo->Settings->BucketId))), ESchemaAttributeVisibility::Public });
 
 		// We retrieve all the session attributes
 		EOS_SessionDetails_GetSessionAttributeCountOptions GetAttributeCountOptions = {};
 		GetAttributeCountOptions.ApiVersion = EOS_SESSIONDETAILS_GETSESSIONATTRIBUTECOUNT_API_LATEST;
 		static_assert(EOS_SESSIONDETAILS_GETSESSIONATTRIBUTECOUNT_API_LATEST == 1, "EOS_SessionDetails_GetSessionAttributeCountOptions updated, check new fields");
 
-		uint32_t AttributeCount = EOS_SessionDetails_GetSessionAttributeCount(InSessionDetailsHandle, &GetAttributeCountOptions);
+		uint32_t AttributeCount = EOS_SessionDetails_GetSessionAttributeCount(SessionDetailsHandle->SessionDetailsHandle, &GetAttributeCountOptions);
 		for (uint32_t Index = 0; Index < AttributeCount; ++Index)
 		{
+			// We parse a single attribute
+
 			EOS_SessionDetails_CopySessionAttributeByIndexOptions CopyAttributeByIndexOptions = {};
 			CopyAttributeByIndexOptions.ApiVersion = EOS_SESSIONDETAILS_COPYSESSIONATTRIBUTEBYINDEX_API_LATEST;
 			static_assert(EOS_SESSIONDETAILS_COPYSESSIONATTRIBUTEBYINDEX_API_LATEST == 1, "EOS_SessionDetails_CopySessionAttributeByIndexOptions updated, check new fields");
@@ -79,71 +81,47 @@ FSessionEOSGS::FSessionEOSGS(const EOS_HSessionDetails& InSessionDetailsHandle)
 			CopyAttributeByIndexOptions.AttrIndex = Index;
 
 			EOS_SessionDetails_Attribute* Attribute = nullptr;
-			EOS_EResult CopyAttributeByIndexResult = EOS_SessionDetails_CopySessionAttributeByIndex(InSessionDetailsHandle, &CopyAttributeByIndexOptions, &Attribute);
+			EOS_EResult CopyAttributeByIndexResult = EOS_SessionDetails_CopySessionAttributeByIndex(SessionDetailsHandle->SessionDetailsHandle, &CopyAttributeByIndexOptions, &Attribute);
 			if (CopyAttributeByIndexResult == EOS_EResult::EOS_Success)
 			{
-				// We parse a single attribute
 				FString Key(Attribute->Data->Key);
-
-				// If the Key contains the ':' character, it will be a Session Member setting
-				if (Key.Contains(TEXT(":")))
+				FSessionAttributeConverter<ESessionAttributeConversionType::FromService> CustomSettingConverter(*Attribute->Data);
+				TPair<FSchemaAttributeId, FSchemaVariant> CustomSettingData = CustomSettingConverter.GetAttributeData();
+					
+				// Most Session Settings values get parsed in the same way as Custom Session Settings, so we will attempt to retrieve them
+				if (Key == EOSGS_OWNER_ACCOUNT_ID_ATTRIBUTE_KEY.ToString())
+				{
+					const FString& PlayerIdStr = CustomSettingData.Value.GetString();
+					const EOS_ProductUserId ProductUserId = EOS_ProductUserId_FromString(TCHAR_TO_UTF8(*PlayerIdStr));
+					OwnerAccountId = FindAccountId(ProductUserId);
+				}
+				if (Key == EOSGS_ALLOW_NEW_MEMBERS_ATTRIBUTE_KEY.ToString())
+				{
+					SessionSettings.bAllowNewMembers = CustomSettingData.Value.GetBoolean();
+				}
+				else if (Key == EOSGS_ANTI_CHEAT_PROTECTED_ATTRIBUTE_KEY.ToString())
+				{
+					SessionInfo.bAntiCheatProtected = CustomSettingData.Value.GetBoolean();
+				}
+				else if (Key == EOSGS_IS_DEDICATED_SERVER_SESSION_ATTRIBUTE_KEY.ToString())
+				{
+					SessionInfo.bIsDedicatedServerSession = CustomSettingData.Value.GetBoolean();
+				}
+				else if (Key == EOSGS_SCHEMA_NAME_ATTRIBUTE_KEY.ToString())
+				{
+					SessionSettings.SchemaName = FSchemaId(CustomSettingData.Value.GetString());
+				}
+				else if (Key == EOSGS_SESSION_ID_OVERRIDE_ATTRIBUTE_KEY.ToString())
+				{
+					SessionInfo.SessionIdOverride = CustomSettingData.Value.GetString();
+				}
+				else // The rest are parsed as a Custom Session Setting
 				{
-					TArray<FString> KeyComponents;
-					Key.ParseIntoArray(KeyComponents, TEXT(":"));
-
-					// We retrieve the member id
-					FString PlayerIdStr = KeyComponents[0];
-					EOS_ProductUserId ProductUserId = EOS_ProductUserId_FromString(TCHAR_TO_UTF8(*PlayerIdStr));
-					FAccountId AccountId = FindAccountId(ProductUserId);
-
-					FSessionMember& SessionMember = SessionMembers.FindOrAdd(AccountId);
-
-					// And add the corresponding custom setting
-					FSchemaAttributeId AttributeId(KeyComponents[1]);
-
-					FSessionAttributeConverter<ESessionAttributeConversionType::FromService> CustomSettingConverter(*Attribute->Data);
-					TPair<FSchemaAttributeId, FSchemaVariant> CustomSettingData = CustomSettingConverter.GetAttributeData();
-
 					FCustomSessionSetting CustomSessionSetting;
 					CustomSessionSetting.Visibility = FromServiceType(Attribute->AdvertisementType);
 					CustomSessionSetting.Data = CustomSettingData.Value;
 
-					SessionMember.MemberSettings.Emplace(AttributeId, CustomSessionSetting);
-				}
-				else 
-				{
-					FSessionAttributeConverter<ESessionAttributeConversionType::FromService> CustomSettingConverter(*Attribute->Data);
-					TPair<FSchemaAttributeId, FSchemaVariant> CustomSettingData = CustomSettingConverter.GetAttributeData();
-					
-					// Most Session Settings values get parsed in the same way as Custom Session Settings, so we will attempt to retrieve them
-					if (Key == EOSGS_ALLOW_NEW_MEMBERS.ToString())
-					{
-						SessionSettings.bAllowNewMembers = CustomSettingData.Value.GetBoolean();
-					}
-					else if (Key == EOSGS_ANTI_CHEAT_PROTECTED.ToString())
-					{
-						SessionInfo.bAntiCheatProtected = CustomSettingData.Value.GetBoolean();
-					}
-					else if (Key == EOSGS_IS_DEDICATED_SERVER_SESSION.ToString())
-					{
-						SessionInfo.bIsDedicatedServerSession = CustomSettingData.Value.GetBoolean();
-					}
-					else if (Key == EOSGS_SCHEMA_NAME.ToString())
-					{
-						SessionSettings.SchemaName = FSchemaId(CustomSettingData.Value.GetString());
-					}
-					else if (Key == EOSGS_SESSION_ID_OVERRIDE.ToString())
-					{
-						SessionInfo.SessionIdOverride = CustomSettingData.Value.GetString();
-					}
-					else // The rest are parsed as a Custom Session Setting
-					{
-						FCustomSessionSetting CustomSessionSetting;
-						CustomSessionSetting.Visibility = FromServiceType(Attribute->AdvertisementType);
-						CustomSessionSetting.Data = CustomSettingData.Value;
-
-						SessionSettings.CustomSettings.Emplace(Key, CustomSessionSetting);
-					}
+					SessionSettings.CustomSettings.Emplace(Key, CustomSessionSetting);
 				}
 			}
 			else
@@ -383,9 +361,9 @@ TOnlineAsyncOpHandle<FCreateSession> FSessionsEOSGS::CreateSession(FCreateSessio
 		}
 
 		// Check if the Bucket Id custom setting is set. EOS Sessions can not be created without it
-		if (!OpParams.SessionSettings.CustomSettings.Contains(EOSGS_BUCKET_ID))
+		if (!OpParams.SessionSettings.CustomSettings.Contains(EOSGS_BUCKET_ID_ATTRIBUTE_KEY))
 		{
-			UE_LOG(LogTemp, Warning, TEXT("[FSessionsEOSGS::CreateSession] Could not create session without Custom Setting 'EOSGS_BUCKET_ID' (FString) set."));
+			UE_LOG(LogTemp, Warning, TEXT("[FSessionsEOSGS::CreateSession] Could not create session without Custom Setting 'EOSGS_BUCKET_ID_ATTRIBUTE_KEY' (FString) set."));
 
 			Op.SetError(Errors::InvalidParams());
 
@@ -401,7 +379,7 @@ TOnlineAsyncOpHandle<FCreateSession> FSessionsEOSGS::CreateSession(FCreateSessio
 		CreateSessionModificationOptions.bPresenceEnabled = OpParams.bPresenceEnabled;
 		CreateSessionModificationOptions.bSanctionsEnabled = !OpParams.bAllowSanctionedPlayers;
 
-		const FCustomSessionSetting* BucketIdSetting = OpParams.SessionSettings.CustomSettings.Find(EOSGS_BUCKET_ID);
+		const FCustomSessionSetting* BucketIdSetting = OpParams.SessionSettings.CustomSettings.Find(EOSGS_BUCKET_ID_ATTRIBUTE_KEY);
 		const FTCHARToUTF8 BucketIdUtf8(BucketIdSetting ? *BucketIdSetting->Data.GetString() : nullptr);
 
 		if (BucketIdUtf8.Length())
@@ -589,16 +567,19 @@ void FSessionsEOSGS::WriteCreateSessionModificationHandle(EOS_HSessionModificati
 	SetInvitesAllowed(SessionModificationHandle, Params.SessionSettings.bAllowNewMembers);
 
 	// We won't copy bIsLANSession since it' irrelevant for EOS Sessions
-	AddAttribute(SessionModificationHandle, EOSGS_ALLOW_NEW_MEMBERS, { FSchemaVariant(Params.SessionSettings.bAllowNewMembers), ESchemaAttributeVisibility::Public });
-	AddAttribute(SessionModificationHandle, EOSGS_ANTI_CHEAT_PROTECTED, { FSchemaVariant(Params.bAntiCheatProtected), ESchemaAttributeVisibility::Public });
-	AddAttribute(SessionModificationHandle, EOSGS_IS_DEDICATED_SERVER_SESSION, { FSchemaVariant(Params.bIsDedicatedServerSession), ESchemaAttributeVisibility::Public });
+	AddAttribute(SessionModificationHandle, EOSGS_ALLOW_NEW_MEMBERS_ATTRIBUTE_KEY, { FSchemaVariant(Params.SessionSettings.bAllowNewMembers), ESchemaAttributeVisibility::Public });
+	AddAttribute(SessionModificationHandle, EOSGS_ANTI_CHEAT_PROTECTED_ATTRIBUTE_KEY, { FSchemaVariant(Params.bAntiCheatProtected), ESchemaAttributeVisibility::Public });
+	AddAttribute(SessionModificationHandle, EOSGS_IS_DEDICATED_SERVER_SESSION_ATTRIBUTE_KEY, { FSchemaVariant(Params.bIsDedicatedServerSession), ESchemaAttributeVisibility::Public });
 
 	SetPermissionLevel(SessionModificationHandle, Params.SessionSettings.JoinPolicy);
 
 	SetMaxPlayers(SessionModificationHandle, Params.SessionSettings.NumMaxConnections);
 
-	AddAttribute(SessionModificationHandle, EOSGS_SCHEMA_NAME, { FSchemaVariant(Params.SessionSettings.SchemaName.ToString()), ESchemaAttributeVisibility::Public });
-	AddAttribute(SessionModificationHandle, EOSGS_SESSION_ID_OVERRIDE, { FSchemaVariant(Params.SessionIdOverride), ESchemaAttributeVisibility::Public });
+	AddAttribute(SessionModificationHandle, EOSGS_SCHEMA_NAME_ATTRIBUTE_KEY, { FSchemaVariant(Params.SessionSettings.SchemaName.ToString()), ESchemaAttributeVisibility::Public });
+	AddAttribute(SessionModificationHandle, EOSGS_SESSION_ID_OVERRIDE_ATTRIBUTE_KEY, { FSchemaVariant(Params.SessionIdOverride), ESchemaAttributeVisibility::Public });
+
+	EOS_ProductUserId OwnerPUID = GetProductUserId(Params.LocalAccountId);
+	AddAttribute(SessionModificationHandle, EOSGS_OWNER_ACCOUNT_ID_ATTRIBUTE_KEY, { FSchemaVariant(LexToString(OwnerPUID)), ESchemaAttributeVisibility::Public });
 
 	// Custom Settings
 
@@ -608,18 +589,11 @@ void FSessionsEOSGS::WriteCreateSessionModificationHandle(EOS_HSessionModificati
 	}
 
 	// BucketId has its own set method on the API
-	const FCustomSessionSetting* NewBucketIdSetting = Params.SessionSettings.CustomSettings.Find(EOSGS_BUCKET_ID);
+	const FCustomSessionSetting* NewBucketIdSetting = Params.SessionSettings.CustomSettings.Find(EOSGS_BUCKET_ID_ATTRIBUTE_KEY);
 	if (NewBucketIdSetting)
 	{
 		SetBucketId(SessionModificationHandle, NewBucketIdSetting->Data.GetString());
-	}
-
-	// Session Member	
-	for (const TPair<FSchemaAttributeId, FCustomSessionSetting>& CustomSettingEntry : Params.SessionMemberData.MemberSettings)
-	{
-		const FName& Key = FName(LexToString(GetProductUserIdChecked(Params.LocalAccountId)) + TEXT(":") + CustomSettingEntry.Key.ToString());
-		AddAttribute(SessionModificationHandle, Key, CustomSettingEntry.Value);
-	}		
+	}	
 }
 
 void FSessionsEOSGS::WriteUpdateSessionModificationHandle(EOS_HSessionModification& SessionModificationHandle, const FSessionSettingsUpdate& NewSettings)
@@ -637,7 +611,7 @@ void FSessionsEOSGS::WriteUpdateSessionModificationHandle(EOS_HSessionModificati
 
 	if (NewSettings.SchemaName.IsSet())
 	{
-		AddAttribute(SessionModificationHandle, EOSGS_SCHEMA_NAME, { FSchemaVariant(NewSettings.SchemaName.GetValue().ToString()), ESchemaAttributeVisibility::Public });
+		AddAttribute(SessionModificationHandle, EOSGS_SCHEMA_NAME_ATTRIBUTE_KEY, { FSchemaVariant(NewSettings.SchemaName.GetValue().ToString()), ESchemaAttributeVisibility::Public });
 	}
 
 	if (NewSettings.JoinPolicy.IsSet())
@@ -651,7 +625,7 @@ void FSessionsEOSGS::WriteUpdateSessionModificationHandle(EOS_HSessionModificati
 	}
 
 	// BucketId has its own set method on the API
-	const FCustomSessionSetting* NewBucketIdSetting = NewSettings.UpdatedCustomSettings.Find(EOSGS_BUCKET_ID);
+	const FCustomSessionSetting* NewBucketIdSetting = NewSettings.UpdatedCustomSettings.Find(EOSGS_BUCKET_ID_ATTRIBUTE_KEY);
 	if (NewBucketIdSetting)
 	{
 		SetBucketId(SessionModificationHandle, NewBucketIdSetting->Data.GetString());
@@ -670,34 +644,19 @@ void FSessionsEOSGS::WriteUpdateSessionModificationHandle(EOS_HSessionModificati
 	}
 }
 
-void FSessionsEOSGS::WriteUpdateSessionModificationHandle(EOS_HSessionModification& SessionModificationHandle, const FAccountId& LocalAccountId, const FSessionMemberUpdate& MemberUpdate)
-{
-	for (const FSchemaAttributeId& CustomSettingEntryKey : MemberUpdate.RemovedMemberSettings)
-	{
-		const FSchemaAttributeId& Key = FSchemaAttributeId(LexToString(GetProductUserIdChecked(LocalAccountId)) + TEXT(":") + CustomSettingEntryKey.ToString());
-		RemoveAttribute(SessionModificationHandle, Key);
-	}
-
-	for (const TPair<FSchemaAttributeId, FCustomSessionSetting>& CustomSettingEntry : MemberUpdate.UpdatedMemberSettings)
-	{
-		const FSchemaAttributeId& Key = FSchemaAttributeId(LexToString(GetProductUserIdChecked(LocalAccountId)) + TEXT(":") + CustomSettingEntry.Key.ToString());
-		AddAttribute(SessionModificationHandle, Key, CustomSettingEntry.Value);
-	}
-}
-
-TFuture<TOnlineResult<FUpdateSessionImpl>> FSessionsEOSGS::UpdateSessionImpl(FUpdateSessionImpl::Params&& Params)
+TFuture<TOnlineResult<FUpdateSessionSettingsImpl>> FSessionsEOSGS::UpdateSessionSettingsImpl(FUpdateSessionSettingsImpl::Params&& Params)
 {
 	// LAN Sessions
 	TOnlineResult<FGetSessionByName> Result = GetSessionByName({ Params.SessionName });
 	if (Result.IsOk() && Result.GetOkValue().Session->GetSessionInfo().bIsLANSession)
 	{
-		return FSessionsLAN::UpdateSessionImpl(MoveTemp(Params));
+		return FSessionsLAN::UpdateSessionSettingsImpl(MoveTemp(Params));
 	}
 
 	// EOSGS Sessions
 
-	TPromise<TOnlineResult<FUpdateSessionImpl>> Promise;
-	TFuture<TOnlineResult<FUpdateSessionImpl>> Future = Promise.GetFuture();
+	TPromise<TOnlineResult<FUpdateSessionSettingsImpl>> Promise;
+	TFuture<TOnlineResult<FUpdateSessionSettingsImpl>> Future = Promise.GetFuture();
 
 	EOS_Sessions_UpdateSessionModificationOptions UpdateSessionModificationOptions = {};
 	UpdateSessionModificationOptions.ApiVersion = EOS_SESSIONS_UPDATESESSIONMODIFICATION_API_LATEST;
@@ -718,19 +677,14 @@ TFuture<TOnlineResult<FUpdateSessionImpl>> FSessionsEOSGS::UpdateSessionImpl(FUp
 	}
 
 	// After creating the session modification handle, we'll add all the updated data to it
-	WriteUpdateSessionModificationHandle(SessionModificationHandle, Params.Mutations.UpdatedSessionSettings);
-
-	for (const TPair<FAccountId, FSessionMemberUpdate>& Entry : Params.Mutations.UpdatedSessionMembers)
-	{
-		WriteUpdateSessionModificationHandle(SessionModificationHandle, Entry.Key, Entry.Value);
-	}
+	WriteUpdateSessionModificationHandle(SessionModificationHandle, Params.Mutations);
 
 	// Whether we update joinability or not will depend on if bAllowNewMembers was set to a new value
 	FUpdateSessionImplEOSGS::Params UpdateSessionImplEOSGSParams =
 	{
 		MakeShared<FSessionModificationHandleEOSGS>(SessionModificationHandle),
-		Params.Mutations.UpdatedSessionSettings.bAllowNewMembers.IsSet() ?
-			FUpdateSessionJoinabilityParams{ Params.SessionName, Params.Mutations.UpdatedSessionSettings.bAllowNewMembers.GetValue() } :
+		Params.Mutations.bAllowNewMembers.IsSet() ?
+			FUpdateSessionJoinabilityParams{ Params.SessionName, Params.Mutations.bAllowNewMembers.GetValue() } :
 			TOptional<FUpdateSessionJoinabilityParams>() // If bAllowNewMembers is not set, the TOptional will be unset too
 	};
 
@@ -750,7 +704,7 @@ TFuture<TOnlineResult<FUpdateSessionImpl>> FSessionsEOSGS::UpdateSessionImpl(FUp
 				(*FoundSession) += SessionUpdateData;
 
 				// We set the result and fire the event
-				Promise.EmplaceValue(FUpdateSessionImpl::Result{ });
+				Promise.EmplaceValue(FUpdateSessionSettingsImpl::Result{ });
 
 				FSessionUpdated SessionUpdatedEvent{ OpParams.SessionName, SessionUpdateData };
 				SessionEvents.OnSessionUpdated.Broadcast(SessionUpdatedEvent);
@@ -1393,7 +1347,7 @@ TOnlineAsyncOpHandle<FJoinSession> FSessionsEOSGS::JoinSession(FJoinSession::Par
 
 		FSessionJoined Event = { { OpParams.LocalAccountId }, FoundSession->GetSessionId() };
 
-			SessionEvents.OnSessionJoined.Broadcast(Event);
+		SessionEvents.OnSessionJoined.Broadcast(Event);
 
 		// A successful join allows the client to server travel, after which AddSessionMember will be called by the engine
 	})
@@ -1484,14 +1438,16 @@ TResult<TArray<EOS_ProductUserId>, FOnlineError> GetProductUserIdsFromEOSGSSessi
 				// We parse a single attribute
 				FString Key(Attribute->Data->Key);
 
-				// If the Key contains the ':' character, it will contain a user id
-				if (Key.Contains(TEXT(":")))
+				if (Key == EOSGS_OWNER_ACCOUNT_ID_ATTRIBUTE_KEY.ToString())
 				{
-					TArray<FString> KeyComponents;
-					Key.ParseIntoArray(KeyComponents, TEXT(":"));
-					const FString& PlayerIdStr = KeyComponents[0];
+					FSessionAttributeConverter<ESessionAttributeConversionType::FromService> CustomSettingConverter(*Attribute->Data);
+					TPair<FSchemaAttributeId, FSchemaVariant> CustomSettingData = CustomSettingConverter.GetAttributeData();
+
+					const FString& PlayerIdStr = CustomSettingData.Value.GetString();
 					const EOS_ProductUserId ProductUserId = EOS_ProductUserId_FromString(TCHAR_TO_UTF8(*PlayerIdStr));
 					Result.AddUnique(ProductUserId);
+
+					break;
 				}
 			}
 			else
@@ -1528,7 +1484,7 @@ TOnlineAsyncOpHandle<FBuildSessionFromDetailsHandle> FSessionsEOSGS::BuildSessio
 			if (TOnlineAsyncOpPtr<FBuildSessionFromDetailsHandle> StrongOp = WeakOp.Pin())
 			{
 				// After all the ids are resolved, we can build the session safely
-				StrongOp->SetResult(FBuildSessionFromDetailsHandle::Result{ StrongOp->GetParams().LocalAccountId, MakeShared<FSessionEOSGS>(Params.SessionDetailsHandleEOSGS->SessionDetailsHandle) });
+				StrongOp->SetResult(FBuildSessionFromDetailsHandle::Result{ StrongOp->GetParams().LocalAccountId, MakeShared<FSessionEOSGS>(Params.SessionDetailsHandleEOSGS) });
 			}
 		});
 	}
diff --git a/Engine/Plugins/Online/Experimental/OnlineServicesEOSGS/Source/Private/Online/SessionsEOSGSTypes.cpp b/Engine/Plugins/Online/Experimental/OnlineServicesEOSGS/Source/Private/Online/SessionsEOSGSTypes.cpp
--- a/Engine/Plugins/Online/Experimental/OnlineServicesEOSGS/Source/Private/Online/SessionsEOSGSTypes.cpp
+++ b/Engine/Plugins/Online/Experimental/OnlineServicesEOSGS/Source/Private/Online/SessionsEOSGSTypes.cpp
@@ -80,7 +80,7 @@ FSessionAttributeConverter<ESessionAttributeConversionType::ToService>::FSession
 }
 
 FSessionAttributeConverter<ESessionAttributeConversionType::ToService>::FSessionAttributeConverter(const TPair<FSchemaAttributeId, FSchemaVariant>& InData)
-	: KeyConverterStorage(InData.Key.ToString())
+	: KeyConverterStorage(*InData.Key.ToString())
 {
 	AttributeData.ApiVersion = EOS_SESSIONS_ATTRIBUTEDATA_API_LATEST;
 	static_assert(EOS_SESSIONS_ATTRIBUTEDATA_API_LATEST == 1, "EOS_Sessions_AttributeData updated, check new fields");
@@ -103,7 +103,7 @@ FSessionAttributeConverter<ESessionAttributeConversionType::ToService>::FSession
 		break;
 	case ESchemaAttributeType::String:
 		AttributeData.ValueType = EOS_ESessionAttributeType::EOS_AT_STRING;
-		ValueConverterStorage.Emplace(InData.Value.GetString());
+		ValueConverterStorage.Emplace(*InData.Value.GetString());
 		AttributeData.Value.AsUtf8 = ValueConverterStorage->Get();
 		break;
 	}
diff --git a/Engine/Plugins/Online/Experimental/OnlineServicesEOSGS/Source/Public/Online/NboSerializerEOSGSSvc.h b/Engine/Plugins/Online/Experimental/OnlineServicesEOSGS/Source/Public/Online/NboSerializerEOSGSSvc.h
--- a/Engine/Plugins/Online/Experimental/OnlineServicesEOSGS/Source/Public/Online/NboSerializerEOSGSSvc.h
+++ b/Engine/Plugins/Online/Experimental/OnlineServicesEOSGS/Source/Public/Online/NboSerializerEOSGSSvc.h
@@ -33,14 +33,13 @@ inline void SerializeToBuffer(FNboSerializeToBuffer& Ar, const FOnlineSessionId&
 	Ar.WriteBinary(Data.GetData(), Data.Num());
 }
 
-inline void SerializeToBuffer(FNboSerializeToBuffer& Packet, const FSessionMembersMap& SessionMembersMap)
+inline void SerializeToBuffer(FNboSerializeToBuffer& Packet, const FSessionMembersSet& SessionMembersSet)
 {
-	Packet << SessionMembersMap.Num();
+	Packet << SessionMembersSet.Num();
 
-	for (const TPair<FAccountId, FSessionMember>& Entry : SessionMembersMap)
+	for (const FAccountId& SessionMember : SessionMembersSet)
 	{
-		SerializeToBuffer(Packet, Entry.Key);
-		NboSerializerCommonSvc::SerializeToBuffer(Packet, Entry.Value);
+		SerializeToBuffer(Packet, SessionMember);
 	}
 }
 
@@ -64,7 +63,7 @@ inline void SerializeFromBuffer(FNboSerializeFromBuffer& Ar, FOnlineSessionId& S
 	SessionId = FOnlineSessionIdRegistryEOSGS::Get().FromReplicationData(Data);
 }
 
-inline void SerializeFromBuffer(FNboSerializeFromBuffer& Packet, FSessionMembersMap& SessionMembersMap)
+inline void SerializeFromBuffer(FNboSerializeFromBuffer& Packet, FSessionMembersSet& SessionMembersSet)
 {
 	int32 NumEntries = 0;
 	Packet >> NumEntries;
@@ -74,10 +73,7 @@ inline void SerializeFromBuffer(FNboSerializeFromBuffer& Packet, FSessionMembers
 		FAccountId Key;
 		SerializeFromBuffer(Packet, Key);
 
-		FSessionMember Value;
-		NboSerializerCommonSvc::SerializeFromBuffer(Packet, Value);
-
-		SessionMembersMap.Emplace(Key, Value);
+		SessionMembersSet.Emplace(Key);
 	}
 }
 
diff --git a/Engine/Plugins/Online/Experimental/OnlineServicesEOSGS/Source/Public/Online/SessionsEOSGS.h b/Engine/Plugins/Online/Experimental/OnlineServicesEOSGS/Source/Public/Online/SessionsEOSGS.h
--- a/Engine/Plugins/Online/Experimental/OnlineServicesEOSGS/Source/Public/Online/SessionsEOSGS.h
+++ b/Engine/Plugins/Online/Experimental/OnlineServicesEOSGS/Source/Public/Online/SessionsEOSGS.h
@@ -85,7 +85,7 @@ public:
 	FSessionEOSGS(const FSessionEOSGS& InSession) = default;
 
 	/** This constructor should only be used by BuildSessionFromDetailsHandle, after all user ids in the session have been resolved. */
-	FSessionEOSGS(const EOS_HSessionDetails& SessionDetailsHandle);
+	FSessionEOSGS(const TSharedPtr<FSessionDetailsHandleEOSGS>& InSessionDetailsHandle);
 
 	static const FSessionEOSGS& Cast(const ISession& InSession);
 
@@ -200,7 +200,7 @@ public:
 	virtual TOnlineAsyncOpHandle<FRemoveSessionMember> RemoveSessionMember(FRemoveSessionMember::Params&& Params) override;
 
 	// FSessionsCommon
-	virtual TFuture<TOnlineResult<FUpdateSessionImpl>> UpdateSessionImpl(FUpdateSessionImpl::Params&& Params) override;
+	virtual TFuture<TOnlineResult<FUpdateSessionSettingsImpl>> UpdateSessionSettingsImpl(FUpdateSessionSettingsImpl::Params&& Params) override;
 
 protected:
 	void RegisterEventHandlers();
@@ -232,11 +232,6 @@ protected:
 	 */
 	void WriteUpdateSessionModificationHandle(EOS_HSessionModification& SessionModificationHandle, const FSessionSettingsUpdate& NewSettings);
 
-	/**
-	 * Writes only the new values for all updated member settings to the SessionModificationHandle
-	 */
-	void WriteUpdateSessionModificationHandle(EOS_HSessionModification& SessionModificationHandle, const FAccountId& LocalAccountId, const FSessionMemberUpdate& MemberUpdate);
-
 	/**
 	 * Writes all relevant values set in the FindSessions parameters into the SessionSearchHandle
 	 */
diff --git a/Engine/Plugins/Online/Experimental/OnlineServicesEOSGS/Source/Public/Online/SessionsEOSGSTypes.h b/Engine/Plugins/Online/Experimental/OnlineServicesEOSGS/Source/Public/Online/SessionsEOSGSTypes.h
--- a/Engine/Plugins/Online/Experimental/OnlineServicesEOSGS/Source/Public/Online/SessionsEOSGSTypes.h
+++ b/Engine/Plugins/Online/Experimental/OnlineServicesEOSGS/Source/Public/Online/SessionsEOSGSTypes.h
@@ -12,13 +12,15 @@
 
 namespace UE::Online {
 
-static FName EOSGS_ALLOW_NEW_MEMBERS = TEXT("EOSGS_ALLOW_NEW_MEMBERS");
-static FName EOSGS_ANTI_CHEAT_PROTECTED = TEXT("EOSGS_ANTI_CHEAT_PROTECTED");
-static FName EOSGS_IS_DEDICATED_SERVER_SESSION = TEXT("EOSGS_IS_DEDICATED_SERVER_SESSION");
-static FName EOSGS_SCHEMA_NAME = TEXT("EOSGS_SCHEMA_NAME");
-static FName EOSGS_SESSION_ID_OVERRIDE = TEXT("EOSGS_SESSION_ID_OVERRIDE");
+static FName EOSGS_ALLOW_NEW_MEMBERS_ATTRIBUTE_KEY = TEXT("EOSGS_ALLOW_NEW_MEMBERS");
+static FName EOSGS_ANTI_CHEAT_PROTECTED_ATTRIBUTE_KEY = TEXT("EOSGS_ANTI_CHEAT_PROTECTED_ATTRIBUTE_KEY");
+static FName EOSGS_IS_DEDICATED_SERVER_SESSION_ATTRIBUTE_KEY = TEXT("EOSGS_IS_DEDICATED_SERVER_SESSION_ATTRIBUTE_KEY");
+static FName EOSGS_SCHEMA_NAME_ATTRIBUTE_KEY = TEXT("EOSGS_SCHEMA_NAME_ATTRIBUTE_KEY");
+static FName EOSGS_SESSION_ID_OVERRIDE_ATTRIBUTE_KEY = TEXT("EOSGS_SESSION_ID_OVERRIDE_ATTRIBUTE_KEY");
 
-static FName EOSGS_BUCKET_ID = TEXT("EOSGS_BUCKET_ID");
+static FName EOSGS_OWNER_ACCOUNT_ID_ATTRIBUTE_KEY = TEXT("EOSGS_OWNER_ACCOUNT_ID_ATTRIBUTE_KEY");
+
+static FName EOSGS_BUCKET_ID_ATTRIBUTE_KEY = TEXT("EOSGS_BUCKET_ID_ATTRIBUTE_KEY");
 
 EOS_EOnlineSessionPermissionLevel ToServiceType(const ESessionJoinPolicy& Value);
 ESessionJoinPolicy FromServiceType(const EOS_EOnlineSessionPermissionLevel& Value);
diff --git a/Engine/Plugins/Online/Experimental/OnlineServicesNull/Source/Public/Online/NboSerializerNullSvc.h b/Engine/Plugins/Online/Experimental/OnlineServicesNull/Source/Public/Online/NboSerializerNullSvc.h
--- a/Engine/Plugins/Online/Experimental/OnlineServicesNull/Source/Public/Online/NboSerializerNullSvc.h
+++ b/Engine/Plugins/Online/Experimental/OnlineServicesNull/Source/Public/Online/NboSerializerNullSvc.h
@@ -33,14 +33,13 @@ inline void SerializeToBuffer(FNboSerializeToBuffer& Ar, const FOnlineSessionId&
 	Ar.WriteBinary(Data.GetData(), Data.Num());
 }
 
-inline void SerializeToBuffer(FNboSerializeToBuffer& Packet, const FSessionMembersMap& SessionMembersMap)
+inline void SerializeToBuffer(FNboSerializeToBuffer& Packet, const FSessionMembersSet& SessionMembersSet)
 {
-	Packet << SessionMembersMap.Num();
+	Packet << SessionMembersSet.Num();
 
-	for (const TPair<FAccountId, FSessionMember>& Entry : SessionMembersMap)
+	for (const FAccountId& SessionMember : SessionMembersSet)
 	{
-		SerializeToBuffer(Packet, Entry.Key);
-		NboSerializerCommonSvc::SerializeToBuffer(Packet, Entry.Value);
+		SerializeToBuffer(Packet, SessionMember);
 	}
 }
 
@@ -76,7 +75,7 @@ inline void SerializeFromBuffer(FNboSerializeFromBuffer& Ar, FOnlineSessionId& S
 	SessionId = FOnlineSessionIdRegistryNull::Get().FromReplicationData(Data);
 }
 
-inline void SerializeFromBuffer(FNboSerializeFromBuffer& Packet, FSessionMembersMap& SessionMembersMap)
+inline void SerializeFromBuffer(FNboSerializeFromBuffer& Packet, FSessionMembersSet& SessionMembersSet)
 {
 	int32 NumEntries = 0;
 	Packet >> NumEntries;
@@ -86,10 +85,7 @@ inline void SerializeFromBuffer(FNboSerializeFromBuffer& Packet, FSessionMembers
 		FAccountId Key;
 		SerializeFromBuffer(Packet, Key);
 
-		FSessionMember Value;
-		NboSerializerCommonSvc::SerializeFromBuffer(Packet, Value);
-
-		SessionMembersMap.Emplace(Key, Value);
+		SessionMembersSet.Emplace(Key);
 	}
 }
 
diff --git a/Engine/Plugins/Online/Experimental/OnlineServicesOSSAdapter/Source/Private/Online/SessionsOSSAdapter.cpp b/Engine/Plugins/Online/Experimental/OnlineServicesOSSAdapter/Source/Private/Online/SessionsOSSAdapter.cpp
--- a/Engine/Plugins/Online/Experimental/OnlineServicesOSSAdapter/Source/Private/Online/SessionsOSSAdapter.cpp
+++ b/Engine/Plugins/Online/Experimental/OnlineServicesOSSAdapter/Source/Private/Online/SessionsOSSAdapter.cpp
@@ -360,14 +360,11 @@ void FSessionsOSSAdapter::Initialize()
 			FSessionUpdate SessionUpdate;
 			if (bJoined)
 			{
-				if (const FSessionMember* SessionMember = FoundSession->GetSessionMembers().Find(TargetAccountId))
-				{
-					SessionUpdate.AddedSessionMembers.Emplace(TargetAccountId, *SessionMember);
-				}
+				SessionUpdate.AddedSessionMembers.Emplace(TargetAccountId);
 			}
 			else
 			{
-				SessionUpdate.RemovedSessionMembers.Add(TargetAccountId);
+				SessionUpdate.RemovedSessionMembers.Emplace(TargetAccountId);
 			}
 
 			const FSessionUpdated Event { SessionName , SessionUpdate };
@@ -376,64 +373,7 @@ void FSessionsOSSAdapter::Initialize()
 		}
 	}));
 
-	SessionsInterface->AddOnSessionParticipantSettingsUpdatedDelegate_Handle(FOnSessionParticipantSettingsUpdatedDelegate::CreateLambda([this](FName SessionName, const FUniqueNetId& TargetUniqueNetId, const FOnlineSessionSettings& UpdatedSettings)
-	{
-		// We won't transmit events for a session or member that doesn't exist
-		const TOnlineResult<FGetSessionByName> GetSessionByNameResult = GetSessionByName({ SessionName });
-		if (GetSessionByNameResult.IsOk())
-		{
-			const TSharedRef<const ISession>& FoundSession = GetSessionByNameResult.GetOkValue().Session;
-
-			const FOnlineServicesOSSAdapter& ServicesOSSAdapter = static_cast<FOnlineServicesOSSAdapter&>(Services);
-			const FUniqueNetIdRef TargetUniqueNetIdRef = TargetUniqueNetId.AsShared();
-			const FAccountId TargetAccountId = ServicesOSSAdapter.GetAccountIdRegistry().FindOrAddHandle(TargetUniqueNetIdRef);
-
-			FSessionUpdate SessionUpdate;
-			if (const FSessionMember* SessionMember = FoundSession->GetSessionMembers().Find(TargetAccountId))
-			{
-				const ::FSessionSettings* MemberSettings = UpdatedSettings.MemberSettings.Find(TargetUniqueNetIdRef);
-				const FCustomSessionSettingsMap UpdatedMemberSettings = GetV2SessionSettings(*MemberSettings);
-
-				FSessionMemberChanges SessionMemberChanges;
-				for (const TPair<FName, FCustomSessionSetting>& Entry : UpdatedMemberSettings)
-				{
-					if (const FCustomSessionSetting* Setting = SessionMember->MemberSettings.Find(Entry.Key))
-					{
-						if (Setting->Data != Entry.Value.Data || Setting->Visibility != Entry.Value.Visibility || Setting->ID != Entry.Value.ID)
-						{
-							FCustomSessionSettingUpdate SettingUpdate;
-							SettingUpdate.OldValue = *Setting;
-							SettingUpdate.NewValue = Entry.Value;
-
-							SessionMemberChanges.ChangedMemberSettings.Emplace(Entry.Key, SettingUpdate);
-						}
-					}
-					else
-					{
-						SessionMemberChanges.AddedMemberSettings.Emplace(Entry.Key, Entry.Value);
-					}
-				}
-
-				for (const TPair<FName, FCustomSessionSetting>& Entry : SessionMember->MemberSettings)
-				{
-					if (!UpdatedMemberSettings.Contains(Entry.Key))
-					{
-						SessionMemberChanges.RemovedMemberSettings.AddUnique(Entry.Key);
-					}
-				}
-
-				SessionUpdate.SessionMembersChanges.Emplace(TargetAccountId, SessionMemberChanges);
-			}
-
-			// TODO: Method to update the interface session with the latest OSS Adapter session data
-			// Done after the update to read the value differences
-
-			const FSessionUpdated Event{ SessionName , SessionUpdate };
-
-			SessionEvents.OnSessionUpdated.Broadcast(Event);
-
-		}
-	}));
+	// Session member settings don't exist in V2, so we won't bind to FOnSessionParticipantSettingsUpdated
 }
 
 void FSessionsOSSAdapter::Shutdown()
@@ -505,10 +445,10 @@ TOnlineAsyncOpHandle<FCreateSession> FSessionsOSSAdapter::CreateSession(FCreateS
 	return Op->GetHandle();
 }
 
-TFuture<TOnlineResult<FUpdateSessionImpl>> FSessionsOSSAdapter::UpdateSessionImpl(FUpdateSessionImpl::Params&& Params)
+TFuture<TOnlineResult<FUpdateSessionSettingsImpl>> FSessionsOSSAdapter::UpdateSessionSettingsImpl(FUpdateSessionSettingsImpl::Params&& Params)
 {
-	TPromise<TOnlineResult<FUpdateSessionImpl>> Promise;
-	TFuture<TOnlineResult<FUpdateSessionImpl>> Future = Promise.GetFuture();
+	TPromise<TOnlineResult<FUpdateSessionSettingsImpl>> Promise;
+	TFuture<TOnlineResult<FUpdateSessionSettingsImpl>> Future = Promise.GetFuture();
 
 	MakeMulticastAdapter(this, SessionsInterface->OnUpdateSessionCompleteDelegates,
 	[this, Promise = MoveTemp(Promise), OpParams = Params](FName SessionName, const bool bWasSuccessful) mutable
@@ -530,7 +470,7 @@ TFuture<TOnlineResult<FUpdateSessionImpl>> FSessionsOSSAdapter::UpdateSessionImp
 			(*FoundSession) += SessionUpdateData;
 
 			// We set the result and fire the event
-			Promise.EmplaceValue(FUpdateSessionImpl::Result{ });
+			Promise.EmplaceValue(FUpdateSessionSettingsImpl::Result{ });
 
 			FSessionUpdated SessionUpdatedEvent{ OpParams.SessionName, SessionUpdateData };
 			SessionEvents.OnSessionUpdated.Broadcast(SessionUpdatedEvent);
@@ -548,8 +488,9 @@ TFuture<TOnlineResult<FUpdateSessionImpl>> FSessionsOSSAdapter::UpdateSessionImp
 
 	// We will update a copy here, and wait until the operation has completed successfully to update our local data
 	FSessionSettings UpdatedV2Settings = FoundSession->SessionSettings;
-	FSessionSettingsChanges SessionSettingsChanges = BuildSessionSettingsChanges(FoundSession, Params.Mutations.UpdatedSessionSettings);
-	UpdatedV2Settings += SessionSettingsChanges;
+	FSessionUpdate SessionUpdate = BuildSessionUpdate(FoundSession, Params.Mutations);
+	check(SessionUpdate.SessionSettingsChanges.IsSet());
+	UpdatedV2Settings += SessionUpdate.SessionSettingsChanges.GetValue();
 
 	FOnlineSessionSettings UpdatedV1Settings = BuildV1SettingsForUpdate(Params.LocalAccountId, FoundSession);
 
@@ -827,28 +768,7 @@ TOnlineAsyncOpHandle<FStartMatchmaking> FSessionsOSSAdapter::StartMatchmaking(FS
 		FOnlineServicesOSSAdapter& ServicesOSSAdapter = static_cast<FOnlineServicesOSSAdapter&>(Services);
 
 		TArray<FSessionMatchmakingUser> MatchMakingUsers;
-
 		FSessionMatchmakingUser NewMatchMakingUser { ServicesOSSAdapter.GetAccountIdRegistry().GetIdValue(OpParams.SessionCreationParameters.LocalAccountId).ToSharedRef() };
-
-		for (const TPair<FSchemaAttributeId, FCustomSessionSetting>& Entry : OpParams.SessionCreationParameters.SessionMemberData.MemberSettings)
-		{
-			switch (Entry.Value.Data.VariantType)
-			{
-			case ESchemaAttributeType::Bool:
-				NewMatchMakingUser.Attributes.Emplace(Entry.Key.ToString(), FVariantData(Entry.Value.Data.GetBoolean()));
-				break;
-			case ESchemaAttributeType::Double:
-				NewMatchMakingUser.Attributes.Emplace(Entry.Key.ToString(), FVariantData(Entry.Value.Data.GetDouble()));
-				break;
-			case ESchemaAttributeType::Int64:
-				NewMatchMakingUser.Attributes.Emplace(Entry.Key.ToString(), FVariantData(Entry.Value.Data.GetInt64()));
-				break;
-			case ESchemaAttributeType::String:
-				NewMatchMakingUser.Attributes.Emplace(Entry.Key.ToString(), FVariantData(Entry.Value.Data.GetString()));
-				break;
-			}
-		}
-
 		MatchMakingUsers.Add(NewMatchMakingUser);
 
 		FOnlineSessionSettings NewSessionSettings = BuildV1SettingsForCreate(OpParams.SessionCreationParameters);
@@ -1211,7 +1131,7 @@ FOnlineSessionSettings FSessionsOSSAdapter::BuildV1SettingsForCreate(const FCrea
 
 	FOnlineServicesOSSAdapter& ServicesOSSAdapter = static_cast<FOnlineServicesOSSAdapter&>(Services);
 
-	Result.MemberSettings.Emplace(ServicesOSSAdapter.GetAccountIdRegistry().GetIdValue(Params.LocalAccountId).ToSharedRef(), GetV1SessionSettings(Params.SessionMemberData.MemberSettings)); // TODO: Pending SchemaVariant work
+	Result.MemberSettings.Add(ServicesOSSAdapter.GetAccountIdRegistry().GetIdValue(Params.LocalAccountId).ToSharedRef());
 
 	return Result;
 }
@@ -1262,9 +1182,9 @@ FOnlineSessionSettings FSessionsOSSAdapter::BuildV1SettingsForUpdate(const FAcco
 	
 	FOnlineServicesOSSAdapter& ServicesOSSAdapter = static_cast<FOnlineServicesOSSAdapter&>(Services);
 
-	for (const TPair<FAccountId, FSessionMember>& Entry : Session->GetSessionMembers())
+	for (const FAccountId& SessionMember : Session->GetSessionMembers())
 	{
-		Result.MemberSettings.Emplace(ServicesOSSAdapter.GetAccountIdRegistry().GetIdValue(Entry.Key).ToSharedRef(), GetV1SessionSettings(Entry.Value.MemberSettings)); // TODO: Pending SchemaVariant work
+		Result.MemberSettings.Add(ServicesOSSAdapter.GetAccountIdRegistry().GetIdValue(SessionMember).ToSharedRef());
 	}
 
 	return Result;
@@ -1290,10 +1210,7 @@ void FSessionsOSSAdapter::WriteV2SessionSettingsFromV1Session(const FOnlineSessi
 	{
 		FAccountId SessionMemberId = ServicesOSSAdapter.GetAccountIdRegistry().FindOrAddHandle(Entry.Key);
 
-		FSessionMember SessionMember;
-		SessionMember.MemberSettings = GetV2SessionSettings(Entry.Value);
-
-		OutSession->SessionMembers.Emplace(SessionMemberId, SessionMember);
+		OutSession->SessionMembers.Emplace(SessionMemberId);
 	}
 }
 
diff --git a/Engine/Plugins/Online/Experimental/OnlineServicesOSSAdapter/Source/Public/Online/SessionsOSSAdapter.h b/Engine/Plugins/Online/Experimental/OnlineServicesOSSAdapter/Source/Public/Online/SessionsOSSAdapter.h
--- a/Engine/Plugins/Online/Experimental/OnlineServicesOSSAdapter/Source/Public/Online/SessionsOSSAdapter.h
+++ b/Engine/Plugins/Online/Experimental/OnlineServicesOSSAdapter/Source/Public/Online/SessionsOSSAdapter.h
@@ -73,7 +73,7 @@ public:
 	virtual TOnlineAsyncOpHandle<FSendSessionInvite> SendSessionInvite(FSendSessionInvite::Params&& Params) override;
 
 	// FSessionsCommon
-	virtual TFuture<TOnlineResult<FUpdateSessionImpl>> UpdateSessionImpl(FUpdateSessionImpl::Params&& Params) override;
+	virtual TFuture<TOnlineResult<FUpdateSessionSettingsImpl>> UpdateSessionSettingsImpl(FUpdateSessionSettingsImpl::Params&& Params) override;
 
 	TOnlineResult<FGetResolvedConnectString> GetResolvedConnectString(const FGetResolvedConnectString::Params& Params);
 
