Integrate schema into lobbies.

diff --git a/Engine/Plugins/Online/Experimental/OnlineServicesEOSGS/Source/Private/Online/LobbiesEOSGS.cpp b/Engine/Plugins/Online/Experimental/OnlineServicesEOSGS/Source/Private/Online/LobbiesEOSGS.cpp
--- a/Engine/Plugins/Online/Experimental/OnlineServicesEOSGS/Source/Private/Online/LobbiesEOSGS.cpp
+++ b/Engine/Plugins/Online/Experimental/OnlineServicesEOSGS/Source/Private/Online/LobbiesEOSGS.cpp
@@ -17,8 +17,6 @@ namespace UE::Online {
 
 static const FString LobbyDataKeyName = TEXT("LobbyData");
 static const FString LobbyDetailsKeyName = TEXT("LobbyDetails");
-static const FString LobbyChangesKeyName = TEXT("LobbyChanges");
-static const FString LobbyMemberChangesKeyName = TEXT("LobbyMemberChanges");
 static const FString LobbyIdStringKeyName = TEXT("LobbyIdString");
 static const FString LobbySearchKeyName = TEXT("LobbySearch");
 
@@ -42,9 +40,8 @@ void FLobbiesEOSGS::Initialize()
 	LobbyPrerequisites = MakeShared<FLobbyPrerequisitesEOS>(FLobbyPrerequisitesEOS{
 		LobbyInterfaceHandle,
 		StaticCastSharedPtr<FAuthEOSGS>(Services.GetAuthInterface()),
-		LobbySchemaRegistry.ToSharedRef(),
-		ServiceSchema.ToSharedRef(),
-		{SDKManager->GetProductName(), GetBuildUniqueId()}});
+		{SDKManager->GetProductName(), GetBuildUniqueId()},
+		SchemaRegistry});
 
 	LobbyDataRegistry = MakeShared<FLobbyDataRegistryEOS>(LobbyPrerequisites.ToSharedRef());
 
@@ -201,74 +198,34 @@ TOnlineAsyncOpHandle<FCreateLobby> FLobbiesEOSGS::CreateLobby(FCreateLobby::Para
 			return MakeFulfilledPromise<void>().GetFuture();
 		}
 	})
-	// Step 6: Set attributes for the lobby creator.
+	// Step 6: Set lobby and creator attributes, change to user lobby privacy setting.
 	.Then([this](TOnlineAsyncOp<FCreateLobby>& InAsyncOp)
 	{
 		const FCreateLobby::Params& Params = InAsyncOp.GetParams();
 		const TSharedRef<FLobbyDataEOS>& LobbyData = GetOpDataChecked<TSharedRef<FLobbyDataEOS>>(InAsyncOp, LobbyDataKeyName);
 
-		// Add owner attributes to operation data.
-		// CreatingMemberData is reused to update the local lobby data and for dispatching notifications once creation has completed.
-		TSharedRef<FClientLobbyMemberDataChanges> LobbyOwnerAttributes = MakeShared<FClientLobbyMemberDataChanges>(FClientLobbyMemberDataChanges{ Params.UserAttributes, {} });
-		TMap<FAccountId, TSharedRef<FClientLobbyMemberDataChanges>> CreatingMemberData;
-		CreatingMemberData.Add(Params.LocalAccountId, LobbyOwnerAttributes);
-		InAsyncOp.Data.Set<TMap<FAccountId, TSharedRef<FClientLobbyMemberDataChanges>>>(LobbyMemberChangesKeyName, MoveTemp(CreatingMemberData));
-
-		// Write user attributes to the lobby.
-		if (!Params.UserAttributes.IsEmpty())
-		{
-			FLobbiesModifyLobbyMemberDataImpl::Params ModifyLobbyMemberDataParams;
-			ModifyLobbyMemberDataParams.LobbyData = LobbyData;
-			ModifyLobbyMemberDataParams.LocalAccountId = Params.LocalAccountId;
-			ModifyLobbyMemberDataParams.Changes = LobbyOwnerAttributes;
-			return ModifyLobbyMemberDataImpl(MoveTemp(ModifyLobbyMemberDataParams));
-		}
-		else
-		{
-			// Skip attribute write when no user attributes are present.
-			return MakeFulfilledPromise<TDefaultErrorResult<FLobbiesModifyLobbyMemberDataImpl>>(FLobbiesModifyLobbyMemberDataImpl::Result{}).GetFuture();
-		}
-	})
-	// Step 7: Handle result
-	.Then([this, DestroyLobbyDuringCreate](TOnlineAsyncOp<FCreateLobby>& InAsyncOp, TDefaultErrorResult<FLobbiesModifyLobbyMemberDataImpl>&& Result)
-	{
-		const FCreateLobby::Params& Params = InAsyncOp.GetParams();
-		TSharedRef<FLobbyDataEOS> LobbyData = GetOpDataChecked<TSharedRef<FLobbyDataEOS>>(InAsyncOp, LobbyDataKeyName);
-
-		if (Result.IsError())
-		{
-			UE_LOG(LogTemp, Warning, TEXT("[FLobbiesEOSGS::CreateLobby] ModifyLobbyMemberDataImpl Failed: User[%s], Lobby[%s], Result[%s]"),
-				*ToLogString(Params.LocalAccountId), *LobbyData->GetLobbyIdString(), *Result.GetErrorValue().GetLogString());
-
-			return DestroyLobbyDuringCreate(InAsyncOp, Params.LocalAccountId, LobbyData->GetLobbyIdString(), MoveTemp(Result.GetErrorValue()));
-		}
-		else
+		FLobbyClientDataPrepareClientChanges::Params PrepareParams;
+		PrepareParams.LocalAccountId = Params.LocalAccountId;
+		PrepareParams.ClientChanges.Attributes = { Params.Attributes, {} };
+		PrepareParams.ClientChanges.MemberAttributes = { Params.UserAttributes, {} };
+		PrepareParams.ClientChanges.JoinPolicy = Params.JoinPolicy;
+		PrepareParams.ClientChanges.LobbySchema = Params.SchemaId;
+		TOnlineResult<FLobbyClientDataPrepareClientChanges> PrepareResult = LobbyData->GetLobbyClientData()->PrepareClientChanges(MoveTemp(PrepareParams));
+		if (PrepareResult.IsError())
 		{
-			return MakeFulfilledPromise<void>().GetFuture();
+			UE_LOG(LogTemp, Warning, TEXT("[FLobbiesEOSGS::CreateLobby] PrepareClientChanges Failed: User[%s], Lobby[%s], Result[%s]"),
+				*ToLogString(Params.LocalAccountId), *LobbyData->GetLobbyIdString(), *PrepareResult.GetErrorValue().GetLogString());
+			InAsyncOp.SetError(MoveTemp(PrepareResult.GetErrorValue()));
+			return MakeFulfilledPromise<TDefaultErrorResult<FLobbiesModifyLobbyDataImpl>>(FLobbiesModifyLobbyDataImpl::Result{}).GetFuture();
 		}
-	})
-	// Step 8: Add lobby attributes and set privacy to user setting.
-	.Then([this](TOnlineAsyncOp<FCreateLobby>& InAsyncOp)
-	{
-		const FCreateLobby::Params& Params = InAsyncOp.GetParams();
-		const TSharedRef<FLobbyDataEOS>& LobbyData = GetOpDataChecked<TSharedRef<FLobbyDataEOS>>(InAsyncOp, LobbyDataKeyName);
-
-		// Track lobby changes to sync client-side cache at operation completion.
-		TSharedRef<FClientLobbyDataChanges> LobbyChanges = MakeShared<FClientLobbyDataChanges>();
-		InAsyncOp.Data.Set<TSharedRef<FClientLobbyDataChanges>>(LobbyChangesKeyName, LobbyChanges);
-		LobbyChanges->MutatedAttributes = Params.Attributes;
-		LobbyChanges->JoinPolicy.Emplace(Params.JoinPolicy);
 
-		// Add lobby attributes.
-		// Set lobby privacy setting to the user provided value.
 		FLobbiesModifyLobbyDataImpl::Params ModifyLobbyDataParams;
 		ModifyLobbyDataParams.LobbyData = LobbyData;
 		ModifyLobbyDataParams.LocalAccountId = Params.LocalAccountId;
-		ModifyLobbyDataParams.Changes = LobbyChanges;
-
+		ModifyLobbyDataParams.ServiceChanges = MoveTemp(PrepareResult.GetOkValue().ServiceChanges);
 		return ModifyLobbyDataImpl(MoveTemp(ModifyLobbyDataParams));
 	})
-	// Step 9: Handle result.
+	// Step 7: Handle result
 	.Then([this, DestroyLobbyDuringCreate](TOnlineAsyncOp<FCreateLobby>& InAsyncOp, TDefaultErrorResult<FLobbiesModifyLobbyDataImpl>&& Result) mutable
 	{
 		const FCreateLobby::Params& Params = InAsyncOp.GetParams();
@@ -286,7 +243,7 @@ TOnlineAsyncOpHandle<FCreateLobby> FLobbiesEOSGS::CreateLobby(FCreateLobby::Para
 			return MakeFulfilledPromise<void>().GetFuture();
 		}
 	})
-	// Step 10: Add the lobby to the active list, apply changes to the cached lobby object, and signal notifications.
+	// Step 8: Add the lobby to the active list, apply changes to the cached lobby object, and signal notifications.
 	.Then([this](TOnlineAsyncOp<FCreateLobby>& InAsyncOp)
 	{
 		const FCreateLobby::Params& Params = InAsyncOp.GetParams();
@@ -295,18 +252,12 @@ TOnlineAsyncOpHandle<FCreateLobby> FLobbiesEOSGS::CreateLobby(FCreateLobby::Para
 		// Mark the lobby active.
 		AddActiveLobby(Params.LocalAccountId, LobbyData);
 
-		// Add member changes to the lobby changes object.
-		const TSharedRef<FClientLobbyDataChanges>& LobbyChanges = GetOpDataChecked<TSharedRef<FClientLobbyDataChanges>>(InAsyncOp, LobbyChangesKeyName);
-		const TMap<FAccountId, TSharedRef<FClientLobbyMemberDataChanges>>& MemberChanges = GetOpDataChecked<TMap<FAccountId, TSharedRef<FClientLobbyMemberDataChanges>>>(InAsyncOp, LobbyMemberChangesKeyName);
-		LobbyChanges->MutatedMembers = MemberChanges;
-		LobbyChanges->LocalName = InAsyncOp.GetParams().LocalName;
-
-		// Make local changes to lobby data and generate notifications.
-		LobbyData->GetClientLobbyData()->ApplyLobbyUpdateFromLocalChanges(MoveTemp(*LobbyChanges), LobbyEvents);
+		// Update local cache and fire events.
+		LobbyData->GetLobbyClientData()->CommitClientChanges({ &LobbyEvents });
 
 		UE_LOG(LogTemp, Log, TEXT("[FLobbiesEOSGS::CreateLobby] Succeeded: User[%s], Lobby[%s]"),
 			*ToLogString(Params.LocalAccountId), *LobbyData->GetLobbyIdString());
-		InAsyncOp.SetResult(FCreateLobby::Result{LobbyData->GetClientLobbyData()->GetPublicDataPtr()});
+		InAsyncOp.SetResult(FCreateLobby::Result{LobbyData->GetLobbyClientData()->GetPublicDataPtr()});
 	})
 	.Enqueue(GetSerialQueue());
 
@@ -409,7 +360,7 @@ TOnlineAsyncOpHandle<FJoinLobby> FLobbiesEOSGS::JoinLobby(FJoinLobby::Params&& I
 	// Check that lobby is compatible with the running game.
 	if (!LobbyDetails->IsBucketCompatible())
 	{
-		UE_LOG(LogTemp, Warning, TEXT("[FLobbiesEOSGS::JoinLobby] Failed: Lobby is not compatible with the application version. Lobby[%s], Bucket[%s], ExpectedBucket[%s]."),
+		UE_LOG(LogTemp, Warning, TEXT("[FLobbiesEOSGS::JoinLobby] Failed: Lobby is not compatible with the client version. Lobby[%s], Bucket[%s], ExpectedBucket[%s]."),
 			*LobbyData->GetLobbyIdString(), *ToLogString(LobbyDetails->GetInfo()->GetBucketId()), *ToLogString(LobbyPrerequisites->BucketId));
 		Op->SetError(Errors::IncompatibleVersion());
 		return Op->GetHandle();
@@ -443,6 +394,7 @@ TOnlineAsyncOpHandle<FJoinLobby> FLobbiesEOSGS::JoinLobby(FJoinLobby::Params&& I
 			UE_LOG(LogTemp, Warning, TEXT("[FLobbiesEOSGS::JoinLobby] EOS_Lobby_JoinLobby Failed: User[%s], Lobby[%s], Result[%s]"),
 				*ToLogString(Params.LocalAccountId), *LobbyData->GetLobbyIdString(), *LexToString(Data->ResultCode));
 			InAsyncOp.SetError(Errors::FromEOSResult(Data->ResultCode));
+			return;
 		}
 	})
 	// Step 3. Set member attributes.
@@ -451,15 +403,26 @@ TOnlineAsyncOpHandle<FJoinLobby> FLobbiesEOSGS::JoinLobby(FJoinLobby::Params&& I
 		const FJoinLobby::Params& Params = InAsyncOp.GetParams();
 		const TSharedRef<FLobbyDataEOS>& LobbyData = GetOpDataChecked<TSharedRef<FLobbyDataEOS>>(InAsyncOp, LobbyDataKeyName);
 
-		FLobbiesModifyLobbyMemberDataImpl::Params ModifyParams;
-		ModifyParams.LobbyData = LobbyData;
-		ModifyParams.LocalAccountId = Params.LocalAccountId;
-		ModifyParams.Changes = MakeShared<FClientLobbyMemberDataChanges>();
-		ModifyParams.Changes->MutatedAttributes = Params.UserAttributes;
-		return ModifyLobbyMemberDataImpl(MoveTemp(ModifyParams));
+		FLobbyClientDataPrepareClientChanges::Params PrepareParams;
+		PrepareParams.LocalAccountId = Params.LocalAccountId;
+		PrepareParams.ClientChanges.MemberAttributes = { Params.UserAttributes, {} };
+		TOnlineResult<FLobbyClientDataPrepareClientChanges> PrepareResult = LobbyData->GetLobbyClientData()->PrepareClientChanges(MoveTemp(PrepareParams));
+		if (PrepareResult.IsError())
+		{
+			UE_LOG(LogTemp, Warning, TEXT("[FLobbiesEOSGS::JoinLobby] PrepareClientChanges Failed: User[%s], Lobby[%s], Result[%s]"),
+				*ToLogString(Params.LocalAccountId), *LobbyData->GetLobbyIdString(), *PrepareResult.GetErrorValue().GetLogString());
+			InAsyncOp.SetError(MoveTemp(PrepareResult.GetErrorValue()));
+			return MakeFulfilledPromise<TDefaultErrorResult<FLobbiesModifyLobbyDataImpl>>(FLobbiesModifyLobbyDataImpl::Result{}).GetFuture();
+		}
+
+		FLobbiesModifyLobbyDataImpl::Params ModifyLobbyDataParams;
+		ModifyLobbyDataParams.LobbyData = LobbyData;
+		ModifyLobbyDataParams.LocalAccountId = Params.LocalAccountId;
+		ModifyLobbyDataParams.ServiceChanges = MoveTemp(PrepareResult.GetOkValue().ServiceChanges);
+		return ModifyLobbyDataImpl(MoveTemp(ModifyLobbyDataParams));
 	})
 	// Step 4. Handle result.
-	.Then([this](TOnlineAsyncOp<FJoinLobby>& InAsyncOp, TDefaultErrorResult<FLobbiesModifyLobbyMemberDataImpl>&& Result)
+	.Then([this](TOnlineAsyncOp<FJoinLobby>& InAsyncOp, TDefaultErrorResult<FLobbiesModifyLobbyDataImpl>&& Result)
 	{
 		if (Result.IsError())
 		{
@@ -513,15 +476,12 @@ TOnlineAsyncOpHandle<FJoinLobby> FLobbiesEOSGS::JoinLobby(FJoinLobby::Params&& I
 		// Mark the lobby active.
 		AddActiveLobby(Params.LocalAccountId, LobbyData);
 
-		// Apply local changes and dispatch notifications.
-		FClientLobbyDataChanges LobbyChanges;
-		LobbyChanges.LocalName = InAsyncOp.GetParams().LocalName;
-		LobbyChanges.MutatedMembers.Add(Params.LocalAccountId, MakeShared<FClientLobbyMemberDataChanges>(FClientLobbyMemberDataChanges{ Params.UserAttributes, {} }));
-		LobbyData->GetClientLobbyData()->ApplyLobbyUpdateFromLocalChanges(MoveTemp(LobbyChanges), LobbyEvents);
+		// Update local cache and fire events.
+		LobbyData->GetLobbyClientData()->CommitClientChanges({ &LobbyEvents });
 
 		UE_LOG(LogTemp, Log, TEXT("[FLobbiesEOSGS::JoinLobby] Succeeded: User[%s], Lobby[%s]"),
 			*ToLogString(Params.LocalAccountId), *LobbyData->GetLobbyIdString());
-		InAsyncOp.SetResult(FJoinLobby::Result{ LobbyData->GetClientLobbyData()->GetPublicDataPtr()});
+		InAsyncOp.SetResult(FJoinLobby::Result{ LobbyData->GetLobbyClientData()->GetPublicDataPtr()});
 	})
 	.Enqueue(GetSerialQueue());
 
@@ -563,6 +523,18 @@ TOnlineAsyncOpHandle<FLeaveLobby> FLobbiesEOSGS::LeaveLobby(FLeaveLobby::Params&
 		const FLeaveLobby::Params& Params = InAsyncOp.GetParams();
 		const TSharedRef<FLobbyDataEOS>& LobbyData = GetOpDataChecked<TSharedRef<FLobbyDataEOS>>(InAsyncOp, LobbyDataKeyName);
 
+		FLobbyClientDataPrepareClientChanges::Params PrepareParams;
+		PrepareParams.LocalAccountId = Params.LocalAccountId;
+		PrepareParams.ClientChanges.LocalUserLeaveReason = ELobbyMemberLeaveReason::Left;
+		TOnlineResult<FLobbyClientDataPrepareClientChanges> PrepareResult = LobbyData->GetLobbyClientData()->PrepareClientChanges(MoveTemp(PrepareParams));
+		if (PrepareResult.IsError())
+		{
+			UE_LOG(LogTemp, Warning, TEXT("[FLobbiesEOSGS::KickLobbyMember] PrepareClientChanges Failed: User[%s], Lobby[%s], Result[%s]"),
+				*ToLogString(Params.LocalAccountId), *LobbyData->GetLobbyIdString(), *PrepareResult.GetErrorValue().GetLogString());
+			InAsyncOp.SetError(MoveTemp(PrepareResult.GetErrorValue()));
+			return MakeFulfilledPromise<TDefaultErrorResult<FLobbiesLeaveLobbyImpl>>(FLobbiesLeaveLobbyImpl::Result{}).GetFuture();
+		}
+
 		FLobbiesLeaveLobbyImpl::Params LeaveParams;
 		LeaveParams.LobbyData = LobbyData;
 		LeaveParams.LocalAccountId = Params.LocalAccountId;
@@ -580,6 +552,8 @@ TOnlineAsyncOpHandle<FLeaveLobby> FLobbiesEOSGS::LeaveLobby(FLeaveLobby::Params&
 
 			// When leaving a lobby consume the error and report success.
 			//InAsyncOp.SetError(Result.GetErrorValue());
+
+			return;
 		}
 	})
 	.Then([this](TOnlineAsyncOp<FLeaveLobby>& InAsyncOp)
@@ -587,17 +561,12 @@ TOnlineAsyncOpHandle<FLeaveLobby> FLobbiesEOSGS::LeaveLobby(FLeaveLobby::Params&
 		const FLeaveLobby::Params& Params = InAsyncOp.GetParams();
 		const TSharedRef<FLobbyDataEOS>& LobbyData = GetOpDataChecked<TSharedRef<FLobbyDataEOS>>(InAsyncOp, LobbyDataKeyName);
 
-		// Remove the user from the local lobby data and dispatch notifications.
-		FClientLobbyDataChanges LobbyChanges;
-		LobbyChanges.LeavingMembers.Add(Params.LocalAccountId, ELobbyMemberLeaveReason::Left);
-		FApplyLobbyUpdateResult ApplyResult = LobbyData->GetClientLobbyData()->ApplyLobbyUpdateFromLocalChanges(MoveTemp(LobbyChanges), LobbyEvents);
+		// Update local cache and fire events.
+		LobbyData->GetLobbyClientData()->CommitClientChanges({ &LobbyEvents });
 
 		// Remove the lobby from the active list for the user.
 		// The lobby data will be cleaned up once all references are removed.
-		for (FAccountId LeavingMember : ApplyResult.LeavingLocalMembers)
-		{
-			RemoveActiveLobby(LeavingMember, LobbyData);
-		}
+		RemoveActiveLobby(Params.LocalAccountId, LobbyData);
 
 		UE_LOG(LogTemp, Log, TEXT("[FLobbiesEOSGS::LeaveLobby] Succeeded: User[%s], Lobby[%s]"),
 			*ToLogString(Params.LocalAccountId), *LobbyData->GetLobbyIdString());
@@ -659,6 +628,7 @@ TOnlineAsyncOpHandle<FInviteLobbyMember> FLobbiesEOSGS::InviteLobbyMember(FInvit
 			UE_LOG(LogTemp, Warning, TEXT("[FLobbiesEOSGS::InviteLobbyMember] InviteLobbyMemberImpl Failed: User[%s], Lobby[%s], Result[%s]"),
 				*ToLogString(Params.LocalAccountId), *LobbyData->GetLobbyIdString(), *Result.GetErrorValue().GetLogString());
 			InAsyncOp.SetError(MoveTemp(Result.GetErrorValue()));
+			return;
 		}
 	})
 	.Then([this](TOnlineAsyncOp<FInviteLobbyMember>& InAsyncOp)
@@ -715,6 +685,7 @@ TOnlineAsyncOpHandle<FDeclineLobbyInvitation> FLobbiesEOSGS::DeclineLobbyInvitat
 			UE_LOG(LogTemp, Warning, TEXT("[FLobbiesEOSGS::DeclineLobbyInvitation] DeclineLobbyInvitationImpl Failed: User[%s], Lobby[%s], Result[%s]"),
 				*ToLogString(Params.LocalAccountId), *LobbyData->GetLobbyIdString(), *Result.GetErrorValue().GetLogString());
 			InAsyncOp.SetError(MoveTemp(Result.GetErrorValue()));
+			return;
 		}
 	})
 	.Then([this](TOnlineAsyncOp<FDeclineLobbyInvitation>& InAsyncOp)
@@ -766,10 +737,22 @@ TOnlineAsyncOpHandle<FKickLobbyMember> FLobbiesEOSGS::KickLobbyMember(FKickLobby
 		const FKickLobbyMember::Params& Params = InAsyncOp.GetParams();
 		const TSharedRef<FLobbyDataEOS>& LobbyData = GetOpDataChecked<TSharedRef<FLobbyDataEOS>>(InAsyncOp, LobbyDataKeyName);
 
+		FLobbyClientDataPrepareClientChanges::Params PrepareParams;
+		PrepareParams.LocalAccountId = Params.LocalAccountId;
+		PrepareParams.ClientChanges.KickedTargetMember = Params.TargetAccountId;
+		TOnlineResult<FLobbyClientDataPrepareClientChanges> PrepareResult = LobbyData->GetLobbyClientData()->PrepareClientChanges(MoveTemp(PrepareParams));
+		if (PrepareResult.IsError())
+		{
+			UE_LOG(LogTemp, Warning, TEXT("[FLobbiesEOSGS::KickLobbyMember] PrepareClientChanges Failed: User[%s], Lobby[%s], Result[%s]"),
+				*ToLogString(Params.LocalAccountId), *LobbyData->GetLobbyIdString(), *PrepareResult.GetErrorValue().GetLogString());
+			InAsyncOp.SetError(MoveTemp(PrepareResult.GetErrorValue()));
+			return MakeFulfilledPromise<TDefaultErrorResult<FLobbiesKickLobbyMemberImpl>>(FLobbiesKickLobbyMemberImpl::Result{}).GetFuture();
+		}
+
 		FLobbiesKickLobbyMemberImpl::Params KickParams;
 		KickParams.LobbyData = LobbyData;
 		KickParams.LocalAccountId = Params.LocalAccountId;
-		KickParams.TargetAccountId = Params.TargetAccountId;
+		KickParams.TargetAccountId = *PrepareResult.GetOkValue().ServiceChanges.KickedTargetMember;
 		return KickLobbyMemberImpl(MoveTemp(KickParams));
 	})
 	.Then([this](TOnlineAsyncOp<FKickLobbyMember>& InAsyncOp, TDefaultErrorResult<FLobbiesKickLobbyMemberImpl>&& Result)
@@ -790,9 +773,7 @@ TOnlineAsyncOpHandle<FKickLobbyMember> FLobbiesEOSGS::KickLobbyMember(FKickLobby
 		const TSharedRef<FLobbyDataEOS>& LobbyData = GetOpDataChecked<TSharedRef<FLobbyDataEOS>>(InAsyncOp, LobbyDataKeyName);
 
 		// Update local cache and fire events.
-		FClientLobbyDataChanges LobbyChanges;
-		LobbyChanges.LeavingMembers.Add(Params.TargetAccountId, ELobbyMemberLeaveReason::Kicked);
-		LobbyData->GetClientLobbyData()->ApplyLobbyUpdateFromLocalChanges(MoveTemp(LobbyChanges), LobbyEvents);
+		LobbyData->GetLobbyClientData()->CommitClientChanges({ &LobbyEvents });
 
 		UE_LOG(LogTemp, Verbose, TEXT("[FLobbiesEOSGS::KickLobbyMember] Succeeded: User[%s], Lobby[%s]"),
 			*ToLogString(Params.LocalAccountId), *LobbyData->GetLobbyIdString());
@@ -838,10 +819,22 @@ TOnlineAsyncOpHandle<FPromoteLobbyMember> FLobbiesEOSGS::PromoteLobbyMember(FPro
 		const FPromoteLobbyMember::Params& Params = InAsyncOp.GetParams();
 		const TSharedRef<FLobbyDataEOS>& LobbyData = GetOpDataChecked<TSharedRef<FLobbyDataEOS>>(InAsyncOp, LobbyDataKeyName);
 
+		FLobbyClientDataPrepareClientChanges::Params PrepareParams;
+		PrepareParams.LocalAccountId = Params.LocalAccountId;
+		PrepareParams.ClientChanges.OwnerAccountId = Params.TargetAccountId;
+		TOnlineResult<FLobbyClientDataPrepareClientChanges> PrepareResult = LobbyData->GetLobbyClientData()->PrepareClientChanges(MoveTemp(PrepareParams));
+		if (PrepareResult.IsError())
+		{
+			UE_LOG(LogTemp, Warning, TEXT("[FLobbiesEOSGS::FPromoteLobbyMember] PrepareClientChanges Failed: User[%s], Lobby[%s], Result[%s]"),
+				*ToLogString(Params.LocalAccountId), *LobbyData->GetLobbyIdString(), *PrepareResult.GetErrorValue().GetLogString());
+			InAsyncOp.SetError(MoveTemp(PrepareResult.GetErrorValue()));
+			return MakeFulfilledPromise<TDefaultErrorResult<FLobbiesPromoteLobbyMemberImpl>>(FLobbiesPromoteLobbyMemberImpl::Result{}).GetFuture();
+		}
+
 		FLobbiesPromoteLobbyMemberImpl::Params PromoteParams;
 		PromoteParams.LobbyData = LobbyData;
 		PromoteParams.LocalAccountId = Params.LocalAccountId;
-		PromoteParams.TargetAccountId = Params.TargetAccountId;
+		PromoteParams.TargetAccountId = *PrepareResult.GetOkValue().ServiceChanges.OwnerAccountId;
 		return PromoteLobbyMemberImpl(MoveTemp(PromoteParams));
 	})
 	.Then([this](TOnlineAsyncOp<FPromoteLobbyMember>& InAsyncOp, TDefaultErrorResult<FLobbiesPromoteLobbyMemberImpl>&& Result)
@@ -854,6 +847,7 @@ TOnlineAsyncOpHandle<FPromoteLobbyMember> FLobbiesEOSGS::PromoteLobbyMember(FPro
 			UE_LOG(LogTemp, Warning, TEXT("[FLobbiesEOSGS::FPromoteLobbyMember] PromoteLobbyMemberImpl Failed: User[%s], Lobby[%s], Result[%s]"),
 				*ToLogString(Params.LocalAccountId), *LobbyData->GetLobbyIdString(), *Result.GetErrorValue().GetLogString());
 			InAsyncOp.SetError(MoveTemp(Result.GetErrorValue()));
+			return;
 		}
 	})
 	.Then([this](TOnlineAsyncOp<FPromoteLobbyMember>& InAsyncOp)
@@ -862,9 +856,7 @@ TOnlineAsyncOpHandle<FPromoteLobbyMember> FLobbiesEOSGS::PromoteLobbyMember(FPro
 		const TSharedRef<FLobbyDataEOS>& LobbyData = GetOpDataChecked<TSharedRef<FLobbyDataEOS>>(InAsyncOp, LobbyDataKeyName);
 
 		// Update local cache and fire events.
-		FClientLobbyDataChanges LobbyChanges;
-		LobbyChanges.OwnerAccountId = Params.TargetAccountId;
-		LobbyData->GetClientLobbyData()->ApplyLobbyUpdateFromLocalChanges(MoveTemp(LobbyChanges), LobbyEvents);
+		LobbyData->GetLobbyClientData()->CommitClientChanges({ &LobbyEvents });
 
 		UE_LOG(LogTemp, Verbose, TEXT("[FLobbiesEOSGS::FPromoteLobbyMember] Succeeded: User[%s], Lobby[%s]"),
 			*ToLogString(Params.LocalAccountId), *LobbyData->GetLobbyIdString());
@@ -910,14 +902,22 @@ TOnlineAsyncOpHandle<FModifyLobbyJoinPolicy> FLobbiesEOSGS::ModifyLobbyJoinPolic
 		const FModifyLobbyJoinPolicy::Params& Params = InAsyncOp.GetParams();
 		const TSharedRef<FLobbyDataEOS>& LobbyData = GetOpDataChecked<TSharedRef<FLobbyDataEOS>>(InAsyncOp, LobbyDataKeyName);
 
-		TSharedRef<FClientLobbyDataChanges> LobbyChanges = MakeShared<FClientLobbyDataChanges>();
-		InAsyncOp.Data.Set<TSharedRef<FClientLobbyDataChanges>>(LobbyChangesKeyName, LobbyChanges);
-		LobbyChanges->JoinPolicy = Params.JoinPolicy;
+		FLobbyClientDataPrepareClientChanges::Params PrepareParams;
+		PrepareParams.LocalAccountId = Params.LocalAccountId;
+		PrepareParams.ClientChanges.JoinPolicy = Params.JoinPolicy;
+		TOnlineResult<FLobbyClientDataPrepareClientChanges> PrepareResult = LobbyData->GetLobbyClientData()->PrepareClientChanges(MoveTemp(PrepareParams));
+		if (PrepareResult.IsError())
+		{
+			UE_LOG(LogTemp, Warning, TEXT("[FLobbiesEOSGS::ModifyLobbyJoinPolicy] PrepareClientChanges Failed: User[%s], Lobby[%s], Result[%s]"),
+				*ToLogString(Params.LocalAccountId), *LobbyData->GetLobbyIdString(), *PrepareResult.GetErrorValue().GetLogString());
+			InAsyncOp.SetError(MoveTemp(PrepareResult.GetErrorValue()));
+			return MakeFulfilledPromise<TDefaultErrorResult<FLobbiesModifyLobbyDataImpl>>(FLobbiesModifyLobbyDataImpl::Result{}).GetFuture();
+		}
 
 		FLobbiesModifyLobbyDataImpl::Params ModifyLobbyDataParams;
 		ModifyLobbyDataParams.LobbyData = LobbyData;
 		ModifyLobbyDataParams.LocalAccountId = Params.LocalAccountId;
-		ModifyLobbyDataParams.Changes = LobbyChanges;
+		ModifyLobbyDataParams.ServiceChanges = MoveTemp(PrepareResult.GetOkValue().ServiceChanges);
 		return ModifyLobbyDataImpl(MoveTemp(ModifyLobbyDataParams));
 	})
 	.Then([this](TOnlineAsyncOp<FModifyLobbyJoinPolicy>& InAsyncOp, TDefaultErrorResult<FLobbiesModifyLobbyDataImpl>&& Result)
@@ -930,16 +930,16 @@ TOnlineAsyncOpHandle<FModifyLobbyJoinPolicy> FLobbiesEOSGS::ModifyLobbyJoinPolic
 			UE_LOG(LogTemp, Warning, TEXT("[FLobbiesEOSGS::ModifyLobbyJoinPolicy] ModifyLobbyDataImpl Failed: User[%s], Lobby[%s], Result[%s]"),
 				*ToLogString(Params.LocalAccountId), *LobbyData->GetLobbyIdString(), *Result.GetErrorValue().GetLogString());
 			InAsyncOp.SetError(MoveTemp(Result.GetErrorValue()));
+			return;
 		}
 	})
 	.Then([this](TOnlineAsyncOp<FModifyLobbyJoinPolicy>& InAsyncOp)
 	{
 		const FModifyLobbyJoinPolicy::Params& Params = InAsyncOp.GetParams();
 		const TSharedRef<FLobbyDataEOS>& LobbyData = GetOpDataChecked<TSharedRef<FLobbyDataEOS>>(InAsyncOp, LobbyDataKeyName);
-		const TSharedRef<FClientLobbyDataChanges>& LobbyChanges = GetOpDataChecked<TSharedRef<FClientLobbyDataChanges>>(InAsyncOp, LobbyChangesKeyName);
 
 		// Update local cache and fire events.
-		LobbyData->GetClientLobbyData()->ApplyLobbyUpdateFromLocalChanges(MoveTemp(*LobbyChanges), LobbyEvents);
+		LobbyData->GetLobbyClientData()->CommitClientChanges({ &LobbyEvents });
 
 		UE_LOG(LogTemp, Verbose, TEXT("[FLobbiesEOSGS::ModifyLobbyJoinPolicy] Succeeded: User[%s], Lobby[%s]"),
 			*ToLogString(Params.LocalAccountId), *LobbyData->GetLobbyIdString());
@@ -985,15 +985,22 @@ TOnlineAsyncOpHandle<FModifyLobbyAttributes> FLobbiesEOSGS::ModifyLobbyAttribute
 		const FModifyLobbyAttributes::Params& Params = InAsyncOp.GetParams();
 		const TSharedRef<FLobbyDataEOS>& LobbyData = GetOpDataChecked<TSharedRef<FLobbyDataEOS>>(InAsyncOp, LobbyDataKeyName);
 
-		TSharedRef<FClientLobbyDataChanges> LobbyChanges = MakeShared<FClientLobbyDataChanges>();
-		InAsyncOp.Data.Set<TSharedRef<FClientLobbyDataChanges>>(LobbyChangesKeyName, LobbyChanges);
-		LobbyChanges->MutatedAttributes = Params.MutatedAttributes;
-		LobbyChanges->ClearedAttributes = Params.ClearedAttributes;
+		FLobbyClientDataPrepareClientChanges::Params PrepareParams;
+		PrepareParams.LocalAccountId = Params.LocalAccountId;
+		PrepareParams.ClientChanges.Attributes = { Params.UpdatedAttributes, Params.RemovedAttributes };
+		TOnlineResult<FLobbyClientDataPrepareClientChanges> PrepareResult = LobbyData->GetLobbyClientData()->PrepareClientChanges(MoveTemp(PrepareParams));
+		if (PrepareResult.IsError())
+		{
+			UE_LOG(LogTemp, Warning, TEXT("[FLobbiesEOSGS::ModifyLobbyAttributes] PrepareClientChanges Failed: User[%s], Lobby[%s], Result[%s]"),
+				*ToLogString(Params.LocalAccountId), *LobbyData->GetLobbyIdString(), *PrepareResult.GetErrorValue().GetLogString());
+			InAsyncOp.SetError(MoveTemp(PrepareResult.GetErrorValue()));
+			return MakeFulfilledPromise<TDefaultErrorResult<FLobbiesModifyLobbyDataImpl>>(FLobbiesModifyLobbyDataImpl::Result{}).GetFuture();
+		}
 
 		FLobbiesModifyLobbyDataImpl::Params ModifyLobbyDataParams;
 		ModifyLobbyDataParams.LobbyData = LobbyDataRegistry->Find(Params.LobbyId);
 		ModifyLobbyDataParams.LocalAccountId = Params.LocalAccountId;
-		ModifyLobbyDataParams.Changes = LobbyChanges;
+		ModifyLobbyDataParams.ServiceChanges = MoveTemp(PrepareResult.GetOkValue().ServiceChanges);
 		return ModifyLobbyDataImpl(MoveTemp(ModifyLobbyDataParams));
 	})
 	.Then([this](TOnlineAsyncOp<FModifyLobbyAttributes>& InAsyncOp, TDefaultErrorResult<FLobbiesModifyLobbyDataImpl>&& Result)
@@ -1006,16 +1013,16 @@ TOnlineAsyncOpHandle<FModifyLobbyAttributes> FLobbiesEOSGS::ModifyLobbyAttribute
 			UE_LOG(LogTemp, Warning, TEXT("[FLobbiesEOSGS::ModifyLobbyAttributes] ModifyLobbyDataImpl Failed: User[%s], Lobby[%s], Result[%s]"),
 				*ToLogString(Params.LocalAccountId), *LobbyData->GetLobbyIdString(), *Result.GetErrorValue().GetLogString());
 			InAsyncOp.SetError(MoveTemp(Result.GetErrorValue()));
+			return;
 		}
 	})
 	.Then([this](TOnlineAsyncOp<FModifyLobbyAttributes>& InAsyncOp)
 	{
 		const FModifyLobbyAttributes::Params& Params = InAsyncOp.GetParams();
 		const TSharedRef<FLobbyDataEOS>& LobbyData = GetOpDataChecked<TSharedRef<FLobbyDataEOS>>(InAsyncOp, LobbyDataKeyName);
-		const TSharedRef<FClientLobbyDataChanges>& LobbyChanges = GetOpDataChecked<TSharedRef<FClientLobbyDataChanges>>(InAsyncOp, LobbyChangesKeyName);
 
 		// Update local cache and fire events.
-		LobbyData->GetClientLobbyData()->ApplyLobbyUpdateFromLocalChanges(MoveTemp(*LobbyChanges), LobbyEvents);
+		LobbyData->GetLobbyClientData()->CommitClientChanges({ &LobbyEvents });
 
 		UE_LOG(LogTemp, Verbose, TEXT("[FLobbiesEOSGS::ModifyLobbyAttributes] Succeeded: User[%s], Lobby[%s]"),
 			*ToLogString(Params.LocalAccountId), *LobbyData->GetLobbyIdString());
@@ -1061,18 +1068,25 @@ TOnlineAsyncOpHandle<FModifyLobbyMemberAttributes> FLobbiesEOSGS::ModifyLobbyMem
 		const FModifyLobbyMemberAttributes::Params& Params = InAsyncOp.GetParams();
 		const TSharedRef<FLobbyDataEOS>& LobbyData = GetOpDataChecked<TSharedRef<FLobbyDataEOS>>(InAsyncOp, LobbyDataKeyName);
 
-		TSharedRef<FClientLobbyMemberDataChanges> LobbyMemberChanges = MakeShared<FClientLobbyMemberDataChanges>();
-		InAsyncOp.Data.Set<TSharedRef<FClientLobbyMemberDataChanges>>(LobbyMemberChangesKeyName, LobbyMemberChanges);
-		LobbyMemberChanges->MutatedAttributes = Params.MutatedAttributes;
-		LobbyMemberChanges->ClearedAttributes = Params.ClearedAttributes;
+		FLobbyClientDataPrepareClientChanges::Params PrepareParams;
+		PrepareParams.LocalAccountId = Params.LocalAccountId;
+		PrepareParams.ClientChanges.MemberAttributes = { Params.UpdatedAttributes, Params.RemovedAttributes };
+		TOnlineResult<FLobbyClientDataPrepareClientChanges> PrepareResult = LobbyData->GetLobbyClientData()->PrepareClientChanges(MoveTemp(PrepareParams));
+		if (PrepareResult.IsError())
+		{
+			UE_LOG(LogTemp, Warning, TEXT("[FLobbiesEOSGS::ModifyLobbyMemberAttributes] PrepareClientChanges Failed: User[%s], Lobby[%s], Result[%s]"),
+				*ToLogString(Params.LocalAccountId), *LobbyData->GetLobbyIdString(), *PrepareResult.GetErrorValue().GetLogString());
+			InAsyncOp.SetError(MoveTemp(PrepareResult.GetErrorValue()));
+			return MakeFulfilledPromise<TDefaultErrorResult<FLobbiesModifyLobbyDataImpl>>(FLobbiesModifyLobbyDataImpl::Result{}).GetFuture();
+		}
 
-		FLobbiesModifyLobbyMemberDataImpl::Params ModifyLobbyMemberDataParams;
-		ModifyLobbyMemberDataParams.LobbyData = LobbyData;
-		ModifyLobbyMemberDataParams.LocalAccountId = Params.LocalAccountId;
-		ModifyLobbyMemberDataParams.Changes = LobbyMemberChanges;
-		return ModifyLobbyMemberDataImpl(MoveTemp(ModifyLobbyMemberDataParams));
+		FLobbiesModifyLobbyDataImpl::Params ModifyLobbyDataParams;
+		ModifyLobbyDataParams.LobbyData = LobbyData;
+		ModifyLobbyDataParams.LocalAccountId = Params.LocalAccountId;
+		ModifyLobbyDataParams.ServiceChanges = MoveTemp(PrepareResult.GetOkValue().ServiceChanges);
+		return ModifyLobbyDataImpl(MoveTemp(ModifyLobbyDataParams));
 	})
-	.Then([this](TOnlineAsyncOp<FModifyLobbyMemberAttributes>& InAsyncOp, TDefaultErrorResult<FLobbiesModifyLobbyMemberDataImpl>&& Result)
+	.Then([this](TOnlineAsyncOp<FModifyLobbyMemberAttributes>& InAsyncOp, TDefaultErrorResult<FLobbiesModifyLobbyDataImpl>&& Result)
 	{
 		if (Result.IsError())
 		{
@@ -1082,18 +1096,16 @@ TOnlineAsyncOpHandle<FModifyLobbyMemberAttributes> FLobbiesEOSGS::ModifyLobbyMem
 			UE_LOG(LogTemp, Warning, TEXT("[FLobbiesEOSGS::ModifyLobbyMemberAttributes] ModifyLobbyMemberDataImpl Failed: User[%s], Lobby[%s], Result[%s]"),
 				*ToLogString(Params.LocalAccountId), *LobbyData->GetLobbyIdString(), *Result.GetErrorValue().GetLogString());
 			InAsyncOp.SetError(MoveTemp(Result.GetErrorValue()));
+			return;
 		}
 	})
 	.Then([this](TOnlineAsyncOp<FModifyLobbyMemberAttributes>& InAsyncOp)
 	{
 		const FModifyLobbyMemberAttributes::Params& Params = InAsyncOp.GetParams();
 		const TSharedRef<FLobbyDataEOS>& LobbyData = GetOpDataChecked<TSharedRef<FLobbyDataEOS>>(InAsyncOp, LobbyDataKeyName);
-		const TSharedRef<FClientLobbyMemberDataChanges>& LobbyMemberChanges = GetOpDataChecked<TSharedRef<FClientLobbyMemberDataChanges>>(InAsyncOp, LobbyMemberChangesKeyName);
 
 		// Update local cache and fire events.
-		FClientLobbyDataChanges LobbyChanges;
-		LobbyChanges.MutatedMembers.Add(Params.LocalAccountId, LobbyMemberChanges);
-		LobbyData->GetClientLobbyData()->ApplyLobbyUpdateFromLocalChanges(MoveTemp(LobbyChanges), LobbyEvents);
+		LobbyData->GetLobbyClientData()->CommitClientChanges({ &LobbyEvents });
 
 		UE_LOG(LogTemp, Verbose, TEXT("[FLobbiesEOSGS::ModifyLobbyMemberAttributes] Succeeded: User[%s], Lobby[%s]"),
 			*ToLogString(Params.LocalAccountId), *LobbyData->GetLobbyIdString());
@@ -1112,7 +1124,7 @@ TOnlineResult<FGetJoinedLobbies> FLobbiesEOSGS::GetJoinedLobbies(FGetJoinedLobbi
 		Result.Lobbies.Reserve(Lobbies->Num());
 		for (const TSharedRef<FLobbyDataEOS>& LobbyDataEOS : *Lobbies)
 		{
-			Result.Lobbies.Emplace(LobbyDataEOS->GetClientLobbyData()->GetPublicDataPtr());
+			Result.Lobbies.Emplace(LobbyDataEOS->GetLobbyClientData()->GetPublicDataPtr());
 		}
 		return TOnlineResult<FGetJoinedLobbies>(MoveTemp(Result));
 	}
@@ -1282,7 +1294,7 @@ void FLobbiesEOSGS::HandleLobbyInviteAccepted(const EOS_Lobby_LobbyInviteAccepte
 				LobbyEvents.OnUILobbyJoinRequested.Broadcast(
 					FUILobbyJoinRequested{
 						Invite->GetReceiver(),
-						TResult<TSharedRef<const FLobby>, FOnlineError>(Invite->GetLobbyData()->GetClientLobbyData()->GetPublicDataPtr()),
+						TResult<TSharedRef<const FLobby>, FOnlineError>(Invite->GetLobbyData()->GetLobbyClientData()->GetPublicDataPtr()),
 						EUILobbyJoinRequestedSource::FromInvitation
 					});
 			}
@@ -1322,7 +1334,7 @@ void FLobbiesEOSGS::HandleJoinLobbyAccepted(const EOS_Lobby_JoinLobbyAcceptedCal
 				LobbyEvents.OnUILobbyJoinRequested.Broadcast(
 					FUILobbyJoinRequested{
 						LocalAccountId,
-						TResult<TSharedRef<const FLobby>, FOnlineError>(Result.GetOkValue()->GetClientLobbyData()->GetPublicDataPtr()),
+						TResult<TSharedRef<const FLobby>, FOnlineError>(Result.GetOkValue()->GetLobbyClientData()->GetPublicDataPtr()),
 						EUILobbyJoinRequestedSource::Unspecified
 					});
 			}
@@ -1352,8 +1364,6 @@ void FLobbiesEOSGS::CheckMetadata()
 	ToLogString(FLobbiesPromoteLobbyMemberImpl::Result());
 	ToLogString(FLobbiesModifyLobbyDataImpl::Params());
 	ToLogString(FLobbiesModifyLobbyDataImpl::Result());
-	ToLogString(FLobbiesModifyLobbyMemberDataImpl::Params());
-	ToLogString(FLobbiesModifyLobbyMemberDataImpl::Result());
 	ToLogString(FLobbiesProcessLobbyNotificationImpl::Params());
 	ToLogString(FLobbiesProcessLobbyNotificationImpl::Result());
 	Meta::VisitFields(FLobbiesLeaveLobbyImpl::Params(), [](const TCHAR* Name, auto& Field) { return false; });
@@ -1370,8 +1380,6 @@ void FLobbiesEOSGS::CheckMetadata()
 	Meta::VisitFields(FLobbiesPromoteLobbyMemberImpl::Result(), [](const TCHAR* Name, auto& Field) { return false; });
 	Meta::VisitFields(FLobbiesModifyLobbyDataImpl::Params(), [](const TCHAR* Name, auto& Field) { return false; });
 	Meta::VisitFields(FLobbiesModifyLobbyDataImpl::Result(), [](const TCHAR* Name, auto& Field) { return false; });
-	Meta::VisitFields(FLobbiesModifyLobbyMemberDataImpl::Params(), [](const TCHAR* Name, auto& Field) { return false; });
-	Meta::VisitFields(FLobbiesModifyLobbyMemberDataImpl::Result(), [](const TCHAR* Name, auto& Field) { return false; });
 	Meta::VisitFields(FLobbiesProcessLobbyNotificationImpl::Params(), [](const TCHAR* Name, auto& Field) { return false; });
 	Meta::VisitFields(FLobbiesProcessLobbyNotificationImpl::Result(), [](const TCHAR* Name, auto& Field) { return false; });
 }
@@ -1478,7 +1486,7 @@ void FLobbiesEOSGS::AddActiveInvite(const TSharedRef<FLobbyInviteDataEOS>& Invit
 			FLobbyInvitationAdded{
 				Invite->GetReceiver(),
 				Invite->GetSender(),
-				Invite->GetLobbyData()->GetClientLobbyData()->GetPublicDataPtr()
+				Invite->GetLobbyData()->GetLobbyClientData()->GetPublicDataPtr()
 			});
 	}
 }
@@ -1491,7 +1499,7 @@ void FLobbiesEOSGS::RemoveActiveInvite(const TSharedRef<FLobbyInviteDataEOS>& In
 		FLobbyInvitationRemoved{
 			Invite->GetReceiver(),
 			Invite->GetSender(),
-			Invite->GetLobbyData()->GetClientLobbyData()->GetPublicDataPtr()
+			Invite->GetLobbyData()->GetLobbyClientData()->GetPublicDataPtr()
 		});
 }
 
@@ -1724,7 +1732,7 @@ TFuture<TDefaultErrorResult<FLobbiesKickLobbyMemberImpl>> FLobbiesEOSGS::KickLob
 		return MakeFulfilledPromise<TDefaultErrorResult<FLobbiesKickLobbyMemberImpl>>(Errors::InvalidParams()).GetFuture();
 	}
 
-	if (Params.LobbyData->GetClientLobbyData()->GetPublicData().OwnerAccountId != Params.LocalAccountId)
+	if (Params.LobbyData->GetLobbyClientData()->GetPublicData().OwnerAccountId != Params.LocalAccountId)
 	{
 		UE_LOG(LogTemp, Warning, TEXT("[FLobbiesEOSGS::KickLobbyMemberImpl] Failed: User is not the lobby owner. User[%s], Lobby[%s]"),
 			*ToLogString(Params.LocalAccountId), *Params.LobbyData->GetLobbyIdString());
@@ -1781,7 +1789,7 @@ TFuture<TDefaultErrorResult<FLobbiesPromoteLobbyMemberImpl>> FLobbiesEOSGS::Prom
 		return MakeFulfilledPromise<TDefaultErrorResult<FLobbiesPromoteLobbyMemberImpl>>(Errors::InvalidParams()).GetFuture();
 	}
 
-	if (Params.LobbyData->GetClientLobbyData()->GetPublicData().OwnerAccountId != Params.LocalAccountId)
+	if (Params.LobbyData->GetLobbyClientData()->GetPublicData().OwnerAccountId != Params.LocalAccountId)
 	{
 		UE_LOG(LogTemp, Warning, TEXT("[FLobbiesEOSGS::PromoteLobbyMemberImpl] Failed: User is not the lobby owner. User[%s], Lobby[%s]"),
 			*ToLogString(Params.LocalAccountId), *Params.LobbyData->GetLobbyIdString());
@@ -1846,18 +1854,11 @@ TFuture<TDefaultErrorResult<FLobbiesModifyLobbyDataImpl>> FLobbiesEOSGS::ModifyL
 		return MakeFulfilledPromise<TDefaultErrorResult<FLobbiesModifyLobbyDataImpl>>(Errors::InvalidParams()).GetFuture();
 	}
 
-	if (Params.LocalAccountId != Params.LobbyData->GetClientLobbyData()->GetPublicData().OwnerAccountId)
-	{
-		UE_LOG(LogTemp, Warning, TEXT("[FLobbiesEOSGS::ModifyLobbyDataImpl] Failed: User is not lobby owner. User[%s], Lobby[%s]"),
-			*ToLogString(Params.LocalAccountId), *Params.LobbyData->GetLobbyIdString());
-		return MakeFulfilledPromise<TDefaultErrorResult<FLobbiesModifyLobbyDataImpl>>(Errors::InvalidParams()).GetFuture();
-	}
-
 	// Start operation.
 	TPromise<TDefaultErrorResult<FLobbiesModifyLobbyDataImpl>> Promise;
 	TFuture<TDefaultErrorResult<FLobbiesModifyLobbyDataImpl>> Future = Promise.GetFuture();
 
-	LobbyDetails->ApplyLobbyDataUpdateFromLocalChanges(Params.LocalAccountId, *Params.Changes)
+	LobbyDetails->ApplyLobbyDataUpdateFromLocalChanges(Params.LocalAccountId, Params.ServiceChanges)
 	.Next([
 		LocalAccountId = Params.LocalAccountId,
 		LobbyData = Params.LobbyData,
@@ -1880,56 +1881,6 @@ TFuture<TDefaultErrorResult<FLobbiesModifyLobbyDataImpl>> FLobbiesEOSGS::ModifyL
 	return Future;
 }
 
-TFuture<TDefaultErrorResult<FLobbiesModifyLobbyMemberDataImpl>> FLobbiesEOSGS::ModifyLobbyMemberDataImpl(FLobbiesModifyLobbyMemberDataImpl::Params&& Params)
-{
-	// Check prerequisites.
-	if (!Services.GetAuthInterface()->IsLoggedIn(Params.LocalAccountId))
-	{
-		UE_LOG(LogTemp, Warning, TEXT("[FLobbiesEOSGS::ModifyLobbyMemberDataImpl] Failed: User is not logged in. User[%s]"), *ToLogString(Params.LocalAccountId));
-		return MakeFulfilledPromise<TDefaultErrorResult<FLobbiesModifyLobbyMemberDataImpl>>(Errors::InvalidUser()).GetFuture();
-	}
-
-	if (!Params.LobbyData.IsValid())
-	{
-		UE_LOG(LogTemp, Warning, TEXT("[FLobbiesEOSGS::ModifyLobbyMemberDataImpl] Failed: No lobby data provided."));
-		return MakeFulfilledPromise<TDefaultErrorResult<FLobbiesModifyLobbyMemberDataImpl>>(Errors::InvalidParams()).GetFuture();
-	}
-
-	TSharedPtr<FLobbyDetailsEOS> LobbyDetails = Params.LobbyData->GetUserLobbyDetails(Params.LocalAccountId);
-	if (!LobbyDetails)
-	{
-		UE_LOG(LogTemp, Warning, TEXT("[FLobbiesEOSGS::ModifyLobbyMemberDataImpl] Failed: Unable to find lobby details for user. User[%s], Lobby[%s]"),
-			*ToLogString(Params.LocalAccountId), *Params.LobbyData->GetLobbyIdString());
-		return MakeFulfilledPromise<TDefaultErrorResult<FLobbiesModifyLobbyMemberDataImpl>>(Errors::InvalidParams()).GetFuture();
-	}
-
-	// Start operation.
-	TPromise<TDefaultErrorResult<FLobbiesModifyLobbyMemberDataImpl>> Promise;
-	TFuture<TDefaultErrorResult<FLobbiesModifyLobbyMemberDataImpl>> Future = Promise.GetFuture();
-
-	LobbyDetails->ApplyLobbyMemberDataUpdateFromLocalChanges(Params.LocalAccountId, *Params.Changes)
-	.Next([
-		LocalAccountId = Params.LocalAccountId,
-		LobbyData = Params.LobbyData,
-		Promise = MoveTemp(Promise)
-	](EOS_EResult EosResult) mutable
-	{
-		if (EosResult != EOS_EResult::EOS_Success && EosResult != EOS_EResult::EOS_NoChange)
-		{
-			UE_LOG(LogTemp, Warning, TEXT("[FLobbiesEOSGS::ModifyLobbyMemberDataImpl] ApplyLobbyMemberDataUpdateFromLocalChanges Failed: User[%s], Lobby[%s], Result[%s]"),
-				*ToLogString(LocalAccountId), *LobbyData->GetLobbyIdString(), *LexToString(EosResult));
-			Promise.EmplaceValue(Errors::FromEOSResult(EosResult));
-			return;
-		}
-
-		UE_LOG(LogTemp, Verbose, TEXT("[FLobbiesEOSGS::ModifyLobbyMemberDataImpl] Succeeded: User[%s], Lobby[%s]"),
-			*ToLogString(LocalAccountId), *LobbyData->GetLobbyIdString());
-		Promise.EmplaceValue(FLobbiesModifyLobbyMemberDataImpl::Result{});
-	});
-
-	return Future;
-}
-
 TOnlineAsyncOpHandle<FLobbiesProcessLobbyNotificationImpl> FLobbiesEOSGS::ProcessLobbyNotificationImplOp(FLobbiesProcessLobbyNotificationImpl::Params&& InParams)
 {
 	TOnlineAsyncOpRef<FLobbiesProcessLobbyNotificationImpl> Op = GetOp<FLobbiesProcessLobbyNotificationImpl>(MoveTemp(InParams));
@@ -1956,7 +1907,7 @@ TOnlineAsyncOpHandle<FLobbiesProcessLobbyNotificationImpl> FLobbiesEOSGS::Proces
 			UE_LOG(LogTemp, Warning, TEXT("[FLobbiesEOSGS::ProcessLobbyNotificationImplOp] Failed: Unable to find active lobby details to process notificaions. Lobby[%s]"),
 				*Params.LobbyData->GetLobbyIdString());
 			InAsyncOp.SetError(Errors::InvalidState());
-			return MakeFulfilledPromise<TDefaultErrorResultInternal<TSharedRef<FClientLobbySnapshot>>>(Errors::InvalidState()).GetFuture();
+			return MakeFulfilledPromise<TDefaultErrorResultInternal<FLobbyServiceSnapshot>>(Errors::InvalidState()).GetFuture();
 		}
 
 		InAsyncOp.Data.Set<TSharedRef<FLobbyDetailsEOS>>(LobbyDetailsKeyName, LobbyDetails.ToSharedRef());
@@ -1964,7 +1915,7 @@ TOnlineAsyncOpHandle<FLobbiesProcessLobbyNotificationImpl> FLobbiesEOSGS::Proces
 		// Fetch lobby snapshot. Fetching the snapshot resolves the account ids of all lobby members in the snapshot.
 		return LobbyDetails->GetLobbySnapshot();
 	})
-	.Then([this](TOnlineAsyncOp<FLobbiesProcessLobbyNotificationImpl>& InAsyncOp, TDefaultErrorResultInternal<TSharedRef<FClientLobbySnapshot>>&& LobbySnapshotResult)
+	.Then([this](TOnlineAsyncOp<FLobbiesProcessLobbyNotificationImpl>& InAsyncOp, TDefaultErrorResultInternal<FLobbyServiceSnapshot>&& LobbySnapshotResult)
 	{
 		const FLobbiesProcessLobbyNotificationImpl::Params& Params = InAsyncOp.GetParams();
 		TSharedRef<FLobbyDetailsEOS> LobbyDetails = GetOpDataChecked<TSharedRef<FLobbyDetailsEOS>>(InAsyncOp, LobbyDetailsKeyName);
@@ -1978,14 +1929,14 @@ TOnlineAsyncOpHandle<FLobbiesProcessLobbyNotificationImpl> FLobbiesEOSGS::Proces
 		}
 
 		// Get member snapshots.
-		TMap<FAccountId, TSharedRef<FClientLobbyMemberSnapshot>> LobbyMemberSnapshots;
+		TMap<FAccountId, FLobbyMemberServiceSnapshot> LobbyMemberSnapshots;
 		LobbyMemberSnapshots.Reserve(Params.MutatedMembers.Num());
 		for (EOS_ProductUserId MutatedMember : Params.MutatedMembers)
 		{
 			const FAccountId MutatedMemberAccountId = FindAccountId(MutatedMember);
 			if (MutatedMemberAccountId.IsValid())
 			{
-				TDefaultErrorResultInternal<TSharedRef<FClientLobbyMemberSnapshot>> MemberSnapshotResult = LobbyDetails->GetLobbyMemberSnapshot(MutatedMemberAccountId);
+				TDefaultErrorResultInternal<FLobbyMemberServiceSnapshot> MemberSnapshotResult = LobbyDetails->GetLobbyMemberSnapshot(MutatedMemberAccountId);
 				if (MemberSnapshotResult.IsError())
 				{
 					UE_LOG(LogTemp, Warning, TEXT("[FLobbiesEOSGS::ProcessLobbyNotificationImplOp] GetLobbyMemberSnapshot Failed. User[%s], Lobby[%s], Result[%s]"),
@@ -2009,12 +1960,23 @@ TOnlineAsyncOpHandle<FLobbiesProcessLobbyNotificationImpl> FLobbiesEOSGS::Proces
 			}
 		}
 
+		// Prepare client state from snapshot.
+		TOnlineResult<FLobbyClientDataPrepareServiceSnapshot> PrepareSnapshotResult =
+			Params.LobbyData->GetLobbyClientData()->PrepareServiceSnapshot({ MoveTemp(LobbySnapshotResult.GetOkValue()), MoveTemp(LobbyMemberSnapshots), MoveTemp(LeavingMemberReason) });
+		if (PrepareSnapshotResult.IsError())
+		{
+			UE_LOG(LogTemp, Warning, TEXT("[FLobbiesEOSGS::ProcessLobbyNotificationImplOp] PrepareServiceSnapshot Failed. Lobby[%s], Result[%s]"),
+				*Params.LobbyData->GetLobbyIdString(), *PrepareSnapshotResult.GetErrorValue().GetLogString());
+			InAsyncOp.SetError(MoveTemp(PrepareSnapshotResult.GetErrorValue()));
+			return;
+		}
+	})
+	.Then([this](TOnlineAsyncOp<FLobbiesProcessLobbyNotificationImpl>& InAsyncOp)
+	{
+		const FLobbiesProcessLobbyNotificationImpl::Params& Params = InAsyncOp.GetParams();
+
 		// Apply updates and fire notifications.
-		FApplyLobbyUpdateResult Result = Params.LobbyData->GetClientLobbyData()->ApplyLobbyUpdateFromServiceSnapshot(
-			MoveTemp(*LobbySnapshotResult.GetOkValue()),
-			MoveTemp(LobbyMemberSnapshots),
-			MoveTemp(LeavingMemberReason),
-			&LobbyEvents);
+		FLobbyClientDataCommitServiceSnapshot::Result Result = Params.LobbyData->GetLobbyClientData()->CommitServiceSnapshot({ &LobbyEvents });
 
 		// Remove active users if needed.
 		for (FAccountId LeavingMember : Result.LeavingLocalMembers)
diff --git a/Engine/Plugins/Online/Experimental/OnlineServicesEOSGS/Source/Private/Online/LobbiesEOSGSTypes.cpp b/Engine/Plugins/Online/Experimental/OnlineServicesEOSGS/Source/Private/Online/LobbiesEOSGSTypes.cpp
--- a/Engine/Plugins/Online/Experimental/OnlineServicesEOSGS/Source/Private/Online/LobbiesEOSGSTypes.cpp
+++ b/Engine/Plugins/Online/Experimental/OnlineServicesEOSGS/Source/Private/Online/LobbiesEOSGSTypes.cpp
@@ -9,20 +9,20 @@ namespace UE::Online {
 
 namespace Private {
 
-EOS_EComparisonOp TranslateSearchComparison(ELobbyComparisonOp Op)
+EOS_EComparisonOp TranslateSearchComparison(ESchemaAttributeComparisonOp Op)
 {
 	switch (Op)
 	{
-	default:									checkNoEntry(); // Intentional fallthrough
-	case ELobbyComparisonOp::Equals:			return EOS_EComparisonOp::EOS_CO_EQUAL;
-	case ELobbyComparisonOp::NotEquals:			return EOS_EComparisonOp::EOS_CO_NOTEQUAL;
-	case ELobbyComparisonOp::GreaterThan:		return EOS_EComparisonOp::EOS_CO_GREATERTHAN;
-	case ELobbyComparisonOp::GreaterThanEquals:	return EOS_EComparisonOp::EOS_CO_GREATERTHANOREQUAL;
-	case ELobbyComparisonOp::LessThan:			return EOS_EComparisonOp::EOS_CO_LESSTHAN;
-	case ELobbyComparisonOp::LessThanEquals:	return EOS_EComparisonOp::EOS_CO_LESSTHANOREQUAL;
-	case ELobbyComparisonOp::Near:				return EOS_EComparisonOp::EOS_CO_DISTANCE;
-	case ELobbyComparisonOp::In:				return EOS_EComparisonOp::EOS_CO_ONEOF;
-	case ELobbyComparisonOp::NotIn:				return EOS_EComparisonOp::EOS_CO_NOTANYOF;
+	default:												checkNoEntry(); // Intentional fallthrough
+	case ESchemaAttributeComparisonOp::Equals:				return EOS_EComparisonOp::EOS_CO_EQUAL;
+	case ESchemaAttributeComparisonOp::NotEquals:			return EOS_EComparisonOp::EOS_CO_NOTEQUAL;
+	case ESchemaAttributeComparisonOp::GreaterThan:			return EOS_EComparisonOp::EOS_CO_GREATERTHAN;
+	case ESchemaAttributeComparisonOp::GreaterThanEquals:	return EOS_EComparisonOp::EOS_CO_GREATERTHANOREQUAL;
+	case ESchemaAttributeComparisonOp::LessThan:			return EOS_EComparisonOp::EOS_CO_LESSTHAN;
+	case ESchemaAttributeComparisonOp::LessThanEquals:		return EOS_EComparisonOp::EOS_CO_LESSTHANOREQUAL;
+	case ESchemaAttributeComparisonOp::Near:				return EOS_EComparisonOp::EOS_CO_DISTANCE;
+	case ESchemaAttributeComparisonOp::In:					return EOS_EComparisonOp::EOS_CO_ONEOF;
+	case ESchemaAttributeComparisonOp::NotIn:				return EOS_EComparisonOp::EOS_CO_NOTANYOF;
 
 	// todo:
 	// EOS_EComparisonOp::EOS_CO_ANYOF
@@ -47,45 +47,52 @@ FLobbyBucketIdEOS::FLobbyBucketIdEOS(FString ProductName, int32 ProductVersion)
 }
 
 // Attribute translators.
-FLobbyAttributeTranslator<ELobbyTranslationType::ToService>::FLobbyAttributeTranslator(const TPair<FLobbyAttributeId, FLobbyVariant>& FromAttributeData)
-	: FLobbyAttributeTranslator(FromAttributeData.Key, FromAttributeData.Value)
-{
-}
-
-FLobbyAttributeTranslator<ELobbyTranslationType::ToService>::FLobbyAttributeTranslator(FLobbyAttributeId FromAttributeId, const FLobbyVariant& FromAttributeData)
-	: KeyConverterStorage(*FromAttributeId.ToString())
+FLobbyAttributeTranslator<ELobbyTranslationType::ToService>::FLobbyAttributeTranslator(const FSchemaServiceAttributeData& FromAttributeData)
+	: KeyConverterStorage(*FromAttributeData.Id.ToString().ToLower())
+	, AttributeVisibility(EOS_ELobbyAttributeVisibility::EOS_LAT_PRIVATE)
 {
 	AttributeData.ApiVersion = EOS_LOBBY_ATTRIBUTEDATA_API_LATEST;
 	AttributeData.Key = KeyConverterStorage.Get();
 	static_assert(EOS_LOBBY_ATTRIBUTEDATA_API_LATEST == 1, "EOS_Lobby_AttributeData updated, check new fields");
 
-	if (FromAttributeData.VariantData.IsType<FString>())
+	switch (FromAttributeData.Value.GetType())
 	{
-		ValueConverterStorage.Emplace(*FromAttributeData.VariantData.Get<FString>());
+	case ESchemaAttributeType::String:
+		ValueConverterStorage.Emplace(*FromAttributeData.Value.GetString());
 		AttributeData.ValueType = EOS_ELobbyAttributeType::EOS_AT_STRING;
 		AttributeData.Value.AsUtf8 = ValueConverterStorage->Get();
-	}
-	else if (FromAttributeData.VariantData.IsType<int64>())
-	{
+		break;
+
+	case ESchemaAttributeType::Int64:
 		AttributeData.ValueType = EOS_ELobbyAttributeType::EOS_AT_INT64;
-		AttributeData.Value.AsInt64 = FromAttributeData.VariantData.Get<int64>();
-	}
-	else if (FromAttributeData.VariantData.IsType<double>())
-	{
+		AttributeData.Value.AsInt64 = FromAttributeData.Value.GetInt64();
+		break;
+
+	case ESchemaAttributeType::Double:
 		AttributeData.ValueType = EOS_ELobbyAttributeType::EOS_AT_DOUBLE;
-		AttributeData.Value.AsDouble = FromAttributeData.VariantData.Get<double>();
+		AttributeData.Value.AsDouble = FromAttributeData.Value.GetDouble();
+		break;
+
+	case ESchemaAttributeType::Bool:
+		AttributeData.ValueType = EOS_ELobbyAttributeType::EOS_AT_BOOLEAN;
+		AttributeData.Value.AsBool = FromAttributeData.Value.GetBoolean();
+		break;
+
+	case ESchemaAttributeType::None:
+	default:
+		checkNoEntry();
 	}
-	else if (FromAttributeData.VariantData.IsType<bool>())
+
+	if (EnumHasAnyFlags(FromAttributeData.Flags, ESchemaServiceAttributeFlags::Public))
 	{
-		AttributeData.ValueType = EOS_ELobbyAttributeType::EOS_AT_BOOLEAN;
-		AttributeData.Value.AsBool = FromAttributeData.VariantData.Get<bool>();
+		AttributeVisibility = EOS_ELobbyAttributeVisibility::EOS_LAT_PUBLIC;
 	}
 }
 
 FLobbyAttributeTranslator<ELobbyTranslationType::FromService>::FLobbyAttributeTranslator(const EOS_Lobby_AttributeData& FromAttributeData)
 {
-	FLobbyAttributeId AttributeId(UTF8_TO_TCHAR(FromAttributeData.Key));
-	FLobbyVariant VariantData;
+	FSchemaAttributeId AttributeId(UTF8_TO_TCHAR(FromAttributeData.Key));
+	FSchemaVariant VariantData;
 
 	switch (FromAttributeData.ValueType)
 	{
@@ -110,7 +117,7 @@ FLobbyAttributeTranslator<ELobbyTranslationType::FromService>::FLobbyAttributeTr
 		break;
 	}
 
-	AttributeData = TPair<FLobbyAttributeId, FLobbyVariant>(MoveTemp(AttributeId), MoveTemp(VariantData));
+	AttributeData = TPair<FSchemaAttributeId, FSchemaVariant>(MoveTemp(AttributeId), MoveTemp(VariantData));
 }
 
 FLobbyBucketIdTranslator<ELobbyTranslationType::ToService>::FLobbyBucketIdTranslator(const FLobbyBucketIdEOS& BucketId)
@@ -291,14 +298,14 @@ bool FLobbyDetailsEOS::IsBucketCompatible() const
 	return LobbyDetailsInfo->GetBucketId() == Prerequisites->BucketId;
 }
 
-TFuture<TDefaultErrorResultInternal<TSharedRef<FClientLobbySnapshot>>> FLobbyDetailsEOS::GetLobbySnapshot() const
+TFuture<TDefaultErrorResultInternal<FLobbyServiceSnapshot>> FLobbyDetailsEOS::GetLobbySnapshot() const
 {
 	TSharedPtr<FAuthEOSGS> AuthInterface = Prerequisites->AuthInterface.Pin();
 	if (!AuthInterface)
 	{
 		UE_LOG(LogTemp, Warning, TEXT("[FLobbyDetailsEOS::GetLobbySnapshot] Failed: Auth interface missing. Lobby[%s]"),
 			UTF8_TO_TCHAR(GetInfo()->GetLobbyIdEOS()));
-		return MakeFulfilledPromise<TDefaultErrorResultInternal<TSharedRef<FClientLobbySnapshot>>>(Errors::MissingInterface()).GetFuture();
+		return MakeFulfilledPromise<TDefaultErrorResultInternal<FLobbyServiceSnapshot>>(Errors::MissingInterface()).GetFuture();
 	}
 
 	EOS_LobbyDetails_GetMemberCountOptions GetMemberCountOptions = {};
@@ -320,8 +327,8 @@ TFuture<TDefaultErrorResultInternal<TSharedRef<FClientLobbySnapshot>>> FLobbyDet
 		MemberProductUserIds->Emplace(EOS_LobbyDetails_GetMemberByIndex(LobbyDetailsHandle, &GetMemberByIndexOptions));
 	}
 
-	TPromise<TDefaultErrorResultInternal<TSharedRef<FClientLobbySnapshot>>> Promise;
-	TFuture<TDefaultErrorResultInternal<TSharedRef<FClientLobbySnapshot>>> Future = Promise.GetFuture();
+	TPromise<TDefaultErrorResultInternal<FLobbyServiceSnapshot>> Promise;
+	TFuture<TDefaultErrorResultInternal<FLobbyServiceSnapshot>> Future = Promise.GetFuture();
 
 	// Resolve lobby member product user ids to FAccountId before proceeding.
 	AuthInterface->ResolveAccountIds(AssociatedLocalUser, *MemberProductUserIds)
@@ -335,9 +342,9 @@ TFuture<TDefaultErrorResultInternal<TSharedRef<FClientLobbySnapshot>>> FLobbyDet
 			return;
 		}
 
-		TSharedRef<FClientLobbySnapshot> ClientLobbySnapshot = MakeShared<FClientLobbySnapshot>();
-		ClientLobbySnapshot->MaxMembers = StrongThis->GetInfo()->GetMaxMembers();
-		ClientLobbySnapshot->JoinPolicy = TranslateJoinPolicy(StrongThis->GetInfo()->GetPermissionLevel());
+		FLobbyServiceSnapshot LobbyServiceSnapshot;
+		LobbyServiceSnapshot.MaxMembers = StrongThis->GetInfo()->GetMaxMembers();
+		LobbyServiceSnapshot.JoinPolicy = TranslateJoinPolicy(StrongThis->GetInfo()->GetPermissionLevel());
 
 		// Resolve member info.
 		{
@@ -354,10 +361,10 @@ TFuture<TDefaultErrorResultInternal<TSharedRef<FClientLobbySnapshot>>> FLobbyDet
 
 				if (MemberProductUserId == LobbyOwner)
 				{
-					ClientLobbySnapshot->OwnerAccountId = ResolvedMemberAccountId;
+					LobbyServiceSnapshot.OwnerAccountId = ResolvedMemberAccountId;
 				}
 
-				ClientLobbySnapshot->Members.Add(ResolvedMemberAccountId);
+				LobbyServiceSnapshot.Members.Add(ResolvedMemberAccountId);
 			}
 		}
 
@@ -391,28 +398,27 @@ TFuture<TDefaultErrorResultInternal<TSharedRef<FClientLobbySnapshot>>> FLobbyDet
 				}
 
 				FLobbyAttributeTranslator<ELobbyTranslationType::FromService> AttributeTranslator(*LobbyAttribute->Data);
-				ClientLobbySnapshot->Attributes.Add(MoveTemp(AttributeTranslator.GetMutableAttributeData()));
+				LobbyServiceSnapshot.SchemaServiceSnapshot.Attributes.Add(MoveTemp(AttributeTranslator.GetMutableAttributeData()));
 			}
 		}
 
 		UE_LOG(LogTemp, VeryVerbose, TEXT("[FLobbyDetailsEOS::GetLobbySnapshot] Succeeded: Lobby[%s]"),
 			UTF8_TO_TCHAR(StrongThis->GetInfo()->GetLobbyIdEOS()));
-		Promise.EmplaceValue(MoveTemp(ClientLobbySnapshot));
+		Promise.EmplaceValue(MoveTemp(LobbyServiceSnapshot));
 	});
 
-
 	return Future;
 }
 
-TDefaultErrorResultInternal<TSharedRef<FClientLobbyMemberSnapshot>> FLobbyDetailsEOS::GetLobbyMemberSnapshot(FAccountId MemberAccountId) const
+TDefaultErrorResultInternal<FLobbyMemberServiceSnapshot> FLobbyDetailsEOS::GetLobbyMemberSnapshot(FAccountId MemberAccountId) const
 {
 	EOS_ProductUserId MemberProductUserId = GetProductUserIdChecked(MemberAccountId);
 
-	TSharedRef<FClientLobbyMemberSnapshot> LobbyMemberSnapshot = MakeShared<FClientLobbyMemberSnapshot>();
-	LobbyMemberSnapshot->AccountId = MemberAccountId;
+	FLobbyMemberServiceSnapshot LobbyMemberServcieSnapshot;
+	LobbyMemberServcieSnapshot.AccountId = MemberAccountId;
 	// Todo: 
-	//ClientMemberData->PlatformAccountId;
-	//ClientMemberData->PlatformDisplayName;
+	//LobbyMemberServcieSnapshot.PlatformAccountId;
+	//LobbyMemberServcieSnapshot.PlatformDisplayName;
 
 	// Fetch attributes.
 	{
@@ -441,22 +447,20 @@ TDefaultErrorResultInternal<TSharedRef<FClientLobbyMemberSnapshot>> FLobbyDetail
 			{
 				UE_LOG(LogTemp, Warning, TEXT("[FLobbyDetailsEOS::GetLobbyMemberSnapshot] EOS_LobbyDetails_CopyMemberAttributeByIndex Failed: Member[%s], Lobby[%s], Result[%s]"),
 					*ToLogString(MemberAccountId), UTF8_TO_TCHAR(GetInfo()->GetLobbyIdEOS()), *LexToString(EOSResult));
-				return TDefaultErrorResultInternal<TSharedRef<FClientLobbyMemberSnapshot>>(Errors::FromEOSResult(EOSResult));
+				return TDefaultErrorResultInternal<FLobbyMemberServiceSnapshot>(Errors::FromEOSResult(EOSResult));
 			}
 
 			FLobbyAttributeTranslator<ELobbyTranslationType::FromService> AttributeTranslator(*LobbyAttribute->Data);
-			LobbyMemberSnapshot->Attributes.Add(MoveTemp(AttributeTranslator.GetMutableAttributeData()));
+			LobbyMemberServcieSnapshot.SchemaServiceSnapshot.Attributes.Add(MoveTemp(AttributeTranslator.GetMutableAttributeData()));
 		}
 	}
 
 	UE_LOG(LogTemp, VeryVerbose, TEXT("[FLobbyDetailsEOS::GetLobbyMemberSnapshot] Succeeded: Member[%s], Lobby[%s]"),
 		*ToLogString(MemberAccountId), UTF8_TO_TCHAR(GetInfo()->GetLobbyIdEOS()));
-	return TDefaultErrorResultInternal<TSharedRef<FClientLobbyMemberSnapshot>>(MoveTemp(LobbyMemberSnapshot));
+	return TDefaultErrorResultInternal<FLobbyMemberServiceSnapshot>(MoveTemp(LobbyMemberServcieSnapshot));
 }
 
-TFuture<EOS_EResult> FLobbyDetailsEOS::ApplyLobbyDataUpdateFromLocalChanges(
-	FAccountId LocalAccountId,
-	const FClientLobbyDataChanges& Changes) const
+TFuture<EOS_EResult> FLobbyDetailsEOS::ApplyLobbyDataUpdateFromLocalChanges(FAccountId LocalAccountId, const FLobbyClientServiceChanges& ServiceChanges) const
 {
 	EOS_HLobbyModification LobbyModificationHandle = nullptr;
 
@@ -481,11 +485,11 @@ TFuture<EOS_EResult> FLobbyDetailsEOS::ApplyLobbyDataUpdateFromLocalChanges(
 	}
 
 	// Set lobby join policy.
-	if (Changes.JoinPolicy)
+	if (ServiceChanges.JoinPolicy)
 	{
 		EOS_LobbyModification_SetPermissionLevelOptions SetPermissionOptions = {};
 		SetPermissionOptions.ApiVersion = EOS_LOBBYMODIFICATION_SETPERMISSIONLEVEL_API_LATEST;
-		SetPermissionOptions.PermissionLevel = TranslateJoinPolicy(*Changes.JoinPolicy);
+		SetPermissionOptions.PermissionLevel = TranslateJoinPolicy(*ServiceChanges.JoinPolicy);
 		static_assert(EOS_LOBBYMODIFICATION_SETPERMISSIONLEVEL_API_LATEST == 1, "EOS_LobbyModification_SetPermissionLevelOptions updated, check new fields");
 
 		EOSResult = EOS_LobbyModification_SetPermissionLevel(LobbyModificationHandle, &SetPermissionOptions);
@@ -498,29 +502,29 @@ TFuture<EOS_EResult> FLobbyDetailsEOS::ApplyLobbyDataUpdateFromLocalChanges(
 	}
 
 	// Add attributes.
-	for (const TPair<FLobbyAttributeId, FLobbyVariant>& MutatedAttribute : Changes.MutatedAttributes)
+	for (const TPair<FSchemaAttributeId, FSchemaServiceAttributeData>& UpdatedAttribute : ServiceChanges.UpdatedAttributes)
 	{
-		const FLobbyAttributeTranslator<ELobbyTranslationType::ToService> AttributeTranslator(MutatedAttribute);
+		const FLobbyAttributeTranslator<ELobbyTranslationType::ToService> AttributeTranslator(UpdatedAttribute.Value);
 
 		EOS_LobbyModification_AddAttributeOptions AddAttributeOptions = {};
 		AddAttributeOptions.ApiVersion = EOS_LOBBYMODIFICATION_ADDATTRIBUTE_API_LATEST;
 		AddAttributeOptions.Attribute = &AttributeTranslator.GetAttributeData();
-		AddAttributeOptions.Visibility = EOS_ELobbyAttributeVisibility::EOS_LAT_PUBLIC; // todo - get from schema
+		AddAttributeOptions.Visibility = AttributeTranslator.GetAttributeVisibility();
 		static_assert(EOS_LOBBYMODIFICATION_ADDATTRIBUTE_API_LATEST == 1, "EOS_LobbyModification_AddAttributeOptions updated, check new fields");
 
 		EOSResult = EOS_LobbyModification_AddAttribute(LobbyModificationHandle, &AddAttributeOptions);
 		if (EOSResult != EOS_EResult::EOS_Success)
 		{
 			UE_LOG(LogTemp, Warning, TEXT("[FLobbyDetailsEOS::ApplyLobbyDataUpdateFromLocalChanges] EOS_LobbyModification_AddAttribute Failed: User[%s], Lobby[%s], Key[%s], Value[%s], Result[%s]"),
-				*ToLogString(LocalAccountId), UTF8_TO_TCHAR(GetInfo()->GetLobbyIdEOS()), *MutatedAttribute.Key.ToString().ToLower(), *ToLogString(MutatedAttribute.Value), *LexToString(EOSResult));
+				*ToLogString(LocalAccountId), UTF8_TO_TCHAR(GetInfo()->GetLobbyIdEOS()), *UpdatedAttribute.Key.ToString().ToLower(), *ToLogString(UpdatedAttribute.Value.Value), *LexToString(EOSResult));
 			return MakeFulfilledPromise<EOS_EResult>(EOSResult).GetFuture();
 		}
 	}
 
 	// Remove attributes.
-	for (const FLobbyAttributeId& ClearedAttribute : Changes.ClearedAttributes)
+	for (const FSchemaAttributeId& RemovedAttribute : ServiceChanges.RemovedAttributes)
 	{
-		const FTCHARToUTF8 KeyConverter(*ClearedAttribute.ToString());
+		const FTCHARToUTF8 KeyConverter(*RemovedAttribute.ToString());
 
 		EOS_LobbyModification_RemoveAttributeOptions RemoveAttributeOptions = {};
 		RemoveAttributeOptions.ApiVersion = EOS_LOBBYMODIFICATION_REMOVEATTRIBUTE_API_LATEST;
@@ -531,130 +535,75 @@ TFuture<EOS_EResult> FLobbyDetailsEOS::ApplyLobbyDataUpdateFromLocalChanges(
 		if (EOSResult != EOS_EResult::EOS_Success)
 		{
 			UE_LOG(LogTemp, Warning, TEXT("[FLobbyDetailsEOS::ApplyLobbyDataUpdateFromLocalChanges] EOS_LobbyModification_RemoveAttribute Failed: User[%s], Lobby[%s], Key[%s], Result[%s]"),
-				*ToLogString(LocalAccountId), UTF8_TO_TCHAR(GetInfo()->GetLobbyIdEOS()), UTF8_TO_TCHAR(KeyConverter.Get()), *LexToString(EOSResult));
+				*ToLogString(LocalAccountId), UTF8_TO_TCHAR(GetInfo()->GetLobbyIdEOS()), *RemovedAttribute.ToString().ToLower(), *LexToString(EOSResult));
 			return MakeFulfilledPromise<EOS_EResult>(EOSResult).GetFuture();
 		}
 	}
 
-	TPromise<EOS_EResult> Promise;
-	TFuture<EOS_EResult> Future = Promise.GetFuture();
-
-	// Apply lobby updates.
-	EOS_Lobby_UpdateLobbyOptions UpdateLobbyOptions = {};
-	UpdateLobbyOptions.ApiVersion = EOS_LOBBY_UPDATELOBBY_API_LATEST;
-	UpdateLobbyOptions.LobbyModificationHandle = LobbyModificationHandle;
-	static_assert(EOS_LOBBY_UPDATELOBBY_API_LATEST == 1, "EOS_Lobby_UpdateLobbyOptions updated, check new fields");
-
-	EOS_Async(EOS_Lobby_UpdateLobby, Prerequisites->LobbyInterfaceHandle, UpdateLobbyOptions,
-	[LocalAccountId, LobbyInfo = GetInfo(), Promise = MoveTemp(Promise)](const EOS_Lobby_UpdateLobbyCallbackInfo* Data) mutable
-	{
-		if (Data->ResultCode != EOS_EResult::EOS_Success && Data->ResultCode != EOS_EResult::EOS_NoChange)
-		{
-			UE_LOG(LogTemp, Warning, TEXT("[FLobbyDetailsEOS::ApplyLobbyDataUpdateFromLocalChanges] EOS_Lobby_UpdateLobby Failed: User[%s], Lobby[%s], Result[%s]"),
-				*ToLogString(LocalAccountId), UTF8_TO_TCHAR(LobbyInfo->GetLobbyIdEOS()), *LexToString(Data->ResultCode));
-		}
-		else
-		{
-			UE_LOG(LogTemp, VeryVerbose, TEXT("[FLobbyDetailsEOS::ApplyLobbyDataUpdateFromLocalChanges] Succeeded: User[%s], Lobby[%s], Result[%s]"),
-				*ToLogString(LocalAccountId), UTF8_TO_TCHAR(LobbyInfo->GetLobbyIdEOS()), *LexToString(Data->ResultCode));
-		}
-
-		Promise.EmplaceValue(Data->ResultCode);
-	});
-	return Future;
-}
-
-TFuture<EOS_EResult> FLobbyDetailsEOS::ApplyLobbyMemberDataUpdateFromLocalChanges(
-	FAccountId LocalAccountId,
-	const FClientLobbyMemberDataChanges& Changes) const
-{
-	EOS_HLobbyModification LobbyModificationHandle = {};
-
-	ON_SCOPE_EXIT
-	{
-		EOS_LobbyModification_Release(LobbyModificationHandle);
-	};
-
-	// Create lobby modification handle.
-	EOS_Lobby_UpdateLobbyModificationOptions ModificationOptions = {};
-	ModificationOptions.ApiVersion = EOS_LOBBY_UPDATELOBBYMODIFICATION_API_LATEST;
-	ModificationOptions.LocalUserId = GetProductUserIdChecked(LocalAccountId);
-	ModificationOptions.LobbyId = GetInfo()->GetLobbyIdEOS();
-	static_assert(EOS_LOBBY_UPDATELOBBYMODIFICATION_API_LATEST == 1, "EOS_Lobby_UpdateLobbyModificationOptions updated, check new fields");
-
-	EOS_EResult EOSResult = EOS_Lobby_UpdateLobbyModification(Prerequisites->LobbyInterfaceHandle, &ModificationOptions, &LobbyModificationHandle);
-	if (EOSResult != EOS_EResult::EOS_Success)
-	{
-		UE_LOG(LogTemp, Warning, TEXT("[FLobbyDetailsEOS::ApplyLobbyMemberDataUpdateFromLocalChanges] EOS_Lobby_UpdateLobbyModification Failed: User[%s], Lobby[%s], Result[%s]"),
-			*ToLogString(LocalAccountId), UTF8_TO_TCHAR(GetInfo()->GetLobbyIdEOS()), *LexToString(EOSResult));
-		return MakeFulfilledPromise<EOS_EResult>(EOSResult).GetFuture();
-	}
-
 	// Add member attributes.
-	for (const TPair<FLobbyAttributeId, FLobbyVariant>& MutatedAttribute : Changes.MutatedAttributes)
+	for (const TPair<FSchemaAttributeId, FSchemaServiceAttributeData>& UpdatedAttribute : ServiceChanges.UpdatedMemberAttributes)
 	{
-		const FLobbyAttributeTranslator<ELobbyTranslationType::ToService> AttributeTranslator(MutatedAttribute);
+		const FLobbyAttributeTranslator<ELobbyTranslationType::ToService> AttributeTranslator(UpdatedAttribute.Value);
 
 		EOS_LobbyModification_AddMemberAttributeOptions AddMemberAttributeOptions = {};
 		AddMemberAttributeOptions.ApiVersion = EOS_LOBBYMODIFICATION_ADDMEMBERATTRIBUTE_API_LATEST;
 		AddMemberAttributeOptions.Attribute = &AttributeTranslator.GetAttributeData();
-		AddMemberAttributeOptions.Visibility = EOS_ELobbyAttributeVisibility::EOS_LAT_PUBLIC; // todo - get from schema
+		AddMemberAttributeOptions.Visibility = AttributeTranslator.GetAttributeVisibility();
 		static_assert(EOS_LOBBYMODIFICATION_ADDMEMBERATTRIBUTE_API_LATEST == 1, "EOS_LobbyModification_AddMemberAttributeOptions updated, check new fields");
-		
+
 		EOSResult = EOS_LobbyModification_AddMemberAttribute(LobbyModificationHandle, &AddMemberAttributeOptions);
 		if (EOSResult != EOS_EResult::EOS_Success)
 		{
 			UE_LOG(LogTemp, Warning, TEXT("[FLobbyDetailsEOS::ApplyLobbyMemberDataUpdateFromLocalChanges] EOS_LobbyModification_AddMemberAttribute Failed: User[%s], Lobby[%s], Key[%s], Value[%s], Result[%s]"),
-				*ToLogString(LocalAccountId), UTF8_TO_TCHAR(GetInfo()->GetLobbyIdEOS()), *MutatedAttribute.Key.ToString().ToLower(), *ToLogString(MutatedAttribute.Value), *LexToString(EOSResult));
+				*ToLogString(LocalAccountId), UTF8_TO_TCHAR(GetInfo()->GetLobbyIdEOS()), *UpdatedAttribute.Key.ToString().ToLower(), *ToLogString(UpdatedAttribute.Value.Value), *LexToString(EOSResult));
 			return MakeFulfilledPromise<EOS_EResult>(EOSResult).GetFuture();
 		}
 	}
 
 	// Remove member attributes.
-	for (const FLobbyAttributeId& ClearedAttribute : Changes.ClearedAttributes)
+	for (const FSchemaAttributeId& RemovedAttribute : ServiceChanges.RemovedMemberAttributes)
 	{
-		const FTCHARToUTF8 KeyConverter(*ClearedAttribute.ToString());
+		const FTCHARToUTF8 KeyConverter(*RemovedAttribute.ToString());
 
 		EOS_LobbyModification_RemoveMemberAttributeOptions RemoveMemberAttributeOptions = {};
 		RemoveMemberAttributeOptions.ApiVersion = EOS_LOBBYMODIFICATION_REMOVEMEMBERATTRIBUTE_API_LATEST;
 		RemoveMemberAttributeOptions.Key = KeyConverter.Get();
 		static_assert(EOS_LOBBYMODIFICATION_REMOVEMEMBERATTRIBUTE_API_LATEST == 1, "EOS_LobbyModification_RemoveMemberAttributeOptions updated, check new fields");
-		
+
 		EOSResult = EOS_LobbyModification_RemoveMemberAttribute(LobbyModificationHandle, &RemoveMemberAttributeOptions);
 		if (EOSResult != EOS_EResult::EOS_Success)
 		{
 			UE_LOG(LogTemp, Warning, TEXT("[FLobbyDetailsEOS::ApplyLobbyMemberDataUpdateFromLocalChanges] EOS_LobbyModification_RemoveMemberAttribute Failed: User[%s], Lobby[%s], Key[%s], Result[%s]"),
-				*ToLogString(LocalAccountId), UTF8_TO_TCHAR(GetInfo()->GetLobbyIdEOS()), UTF8_TO_TCHAR(KeyConverter.Get()), *LexToString(EOSResult));
+				*ToLogString(LocalAccountId), UTF8_TO_TCHAR(GetInfo()->GetLobbyIdEOS()), *RemovedAttribute.ToString().ToLower(), *LexToString(EOSResult));
 			return MakeFulfilledPromise<EOS_EResult>(EOSResult).GetFuture();
 		}
 	}
 
+	TPromise<EOS_EResult> Promise;
+	TFuture<EOS_EResult> Future = Promise.GetFuture();
+
 	// Apply lobby updates.
 	EOS_Lobby_UpdateLobbyOptions UpdateLobbyOptions = {};
 	UpdateLobbyOptions.ApiVersion = EOS_LOBBY_UPDATELOBBY_API_LATEST;
 	UpdateLobbyOptions.LobbyModificationHandle = LobbyModificationHandle;
 	static_assert(EOS_LOBBY_UPDATELOBBY_API_LATEST == 1, "EOS_Lobby_UpdateLobbyOptions updated, check new fields");
 
-	TPromise<EOS_EResult> Promise;
-	TFuture<EOS_EResult> Future = Promise.GetFuture();
-
 	EOS_Async(EOS_Lobby_UpdateLobby, Prerequisites->LobbyInterfaceHandle, UpdateLobbyOptions,
 	[LocalAccountId, LobbyInfo = GetInfo(), Promise = MoveTemp(Promise)](const EOS_Lobby_UpdateLobbyCallbackInfo* Data) mutable
 	{
 		if (Data->ResultCode != EOS_EResult::EOS_Success && Data->ResultCode != EOS_EResult::EOS_NoChange)
 		{
-			UE_LOG(LogTemp, Warning, TEXT("[FLobbyDetailsEOS::ApplyLobbyMemberDataUpdateFromLocalChanges] EOS_Lobby_UpdateLobby Failed: User[%s], Lobby[%s], Result[%s]"),
+			UE_LOG(LogTemp, Warning, TEXT("[FLobbyDetailsEOS::ApplyLobbyDataUpdateFromLocalChanges] EOS_Lobby_UpdateLobby Failed: User[%s], Lobby[%s], Result[%s]"),
 				*ToLogString(LocalAccountId), UTF8_TO_TCHAR(LobbyInfo->GetLobbyIdEOS()), *LexToString(Data->ResultCode));
 		}
 		else
 		{
-			UE_LOG(LogTemp, VeryVerbose, TEXT("[FLobbyDetailsEOS::ApplyLobbyMemberDataUpdateFromLocalChanges] Succeeded: User[%s], Lobby[%s], Result[%s]"),
+			UE_LOG(LogTemp, VeryVerbose, TEXT("[FLobbyDetailsEOS::ApplyLobbyDataUpdateFromLocalChanges] Succeeded: User[%s], Lobby[%s], Result[%s]"),
 				*ToLogString(LocalAccountId), UTF8_TO_TCHAR(LobbyInfo->GetLobbyIdEOS()), *LexToString(Data->ResultCode));
 		}
 
 		Promise.EmplaceValue(Data->ResultCode);
 	});
-
 	return Future;
 }
 
@@ -707,7 +656,7 @@ FLobbyDataEOS::~FLobbyDataEOS()
 {
 	if (UnregisterFn)
 	{
-		UnregisterFn(ClientLobbyData->GetPublicData().LobbyId);
+		UnregisterFn(LobbyClientData->GetPublicData().LobbyId);
 	}
 }
 
@@ -747,10 +696,10 @@ TSharedPtr<FLobbyDetailsEOS> FLobbyDataEOS::GetActiveLobbyDetails() const
 }
 
 FLobbyDataEOS::FLobbyDataEOS(
-	const TSharedRef<FClientLobbyData>& ClientLobbyData,
+	const TSharedRef<FLobbyClientData>& LobbyClientData,
 	const TSharedRef<FLobbyDetailsInfoEOS>& LobbyDetailsInfo,
 	FUnregisterFn UnregisterFn)
-	: ClientLobbyData(ClientLobbyData)
+	: LobbyClientData(LobbyClientData)
 	, LobbyDetailsInfo(LobbyDetailsInfo)
 	, UnregisterFn(MoveTemp(UnregisterFn))
 	, LobbyIdString(UTF8_TO_TCHAR(LobbyDetailsInfo->GetLobbyIdEOS()))
@@ -758,6 +707,7 @@ FLobbyDataEOS::FLobbyDataEOS(
 }
 
 TFuture<TDefaultErrorResultInternal<TSharedRef<FLobbyDataEOS>>> FLobbyDataEOS::Create(
+	TSharedRef<FLobbyPrerequisitesEOS> Prerequisites,
 	FLobbyId LobbyId,
 	const TSharedRef<FLobbyDetailsEOS>& LobbyDetails,
 	FUnregisterFn UnregisterFn)
@@ -771,9 +721,10 @@ TFuture<TDefaultErrorResultInternal<TSharedRef<FLobbyDataEOS>>> FLobbyDataEOS::C
 		Promise = MoveTemp(Promise),
 		LobbyId,
 		LobbyDetails,
+		Prerequisites,
 		UnregisterFn = MoveTemp(UnregisterFn)
 	]
-	(TDefaultErrorResultInternal<TSharedRef<FClientLobbySnapshot>>&& Result) mutable
+	(TDefaultErrorResultInternal<FLobbyServiceSnapshot>&& Result) mutable
 	{
 		if (Result.IsError())
 		{
@@ -783,14 +734,14 @@ TFuture<TDefaultErrorResultInternal<TSharedRef<FLobbyDataEOS>>> FLobbyDataEOS::C
 			return;
 		}
 
-		TSharedRef<FClientLobbySnapshot> LobbySnapshot = Result.GetOkValue();
-		TSharedRef<FClientLobbyData> NewClientLobbyData = MakeShared<FClientLobbyData>(LobbyId);
+		FLobbyServiceSnapshot LobbySnapshot = MoveTemp(Result.GetOkValue());
+		TSharedRef<FLobbyClientData> NewLobbyClientData = MakeShared<FLobbyClientData>(LobbyId, Prerequisites->SchemaRegistry);
 
 		// Fetch member data and apply them to the lobby.
-		TMap<FAccountId, TSharedRef<FClientLobbyMemberSnapshot>> MemberSnapshots;
-		for (FAccountId MemberAccountId : LobbySnapshot->Members)
+		TMap<FAccountId, FLobbyMemberServiceSnapshot> MemberSnapshots;
+		for (FAccountId MemberAccountId : LobbySnapshot.Members)
 		{
-			TDefaultErrorResultInternal<TSharedRef<FClientLobbyMemberSnapshot>> LobbyMemberSnapshotResult = LobbyDetails->GetLobbyMemberSnapshot(MemberAccountId);
+			TDefaultErrorResultInternal<FLobbyMemberServiceSnapshot> LobbyMemberSnapshotResult = LobbyDetails->GetLobbyMemberSnapshot(MemberAccountId);
 			if (LobbyMemberSnapshotResult.IsError())
 			{
 				UE_LOG(LogTemp, Warning, TEXT("[FLobbyDataEOS::Create] FLobbyDetailsEOS::GetLobbyMemberSnapshot Failed: Lobby[%s], User[%s], Result[%s]"),
@@ -802,12 +753,23 @@ TFuture<TDefaultErrorResultInternal<TSharedRef<FLobbyDataEOS>>> FLobbyDataEOS::C
 			MemberSnapshots.Emplace(MemberAccountId, MoveTemp(LobbyMemberSnapshotResult.GetOkValue()));
 		}
 
-		NewClientLobbyData->ApplyLobbyUpdateFromServiceSnapshot(MoveTemp(*LobbySnapshot), MoveTemp(MemberSnapshots));
+		TOnlineResult<FLobbyClientDataPrepareServiceSnapshot> PrepareServiceLobbySnapshotResult =
+			NewLobbyClientData->PrepareServiceSnapshot({ MoveTemp(LobbySnapshot), MoveTemp(MemberSnapshots), {} });
+
+		if (PrepareServiceLobbySnapshotResult.IsError())
+		{
+			UE_LOG(LogTemp, Warning, TEXT("[FLobbyDataEOS::Create] FLobbyClientData::PrepareServiceSnapshot Failed: Lobby[%s], Result[%s]"),
+				*ToLogString(LobbyId), *PrepareServiceLobbySnapshotResult.GetErrorValue().GetLogString());
+			Promise.EmplaceValue(MoveTemp(PrepareServiceLobbySnapshotResult.GetErrorValue()));
+			return;
+		}
+
+		NewLobbyClientData->CommitServiceSnapshot({});
 
 		UE_LOG(LogTemp, VeryVerbose, TEXT("[FLobbyDataEOS::Create] Succeeded: Lobby[%s]"),
 			UTF8_TO_TCHAR(LobbyDetails->GetInfo()->GetLobbyIdEOS()));
 		Promise.EmplaceValue(MakeShared<FLobbyDataEOS>(
-			NewClientLobbyData,
+			NewLobbyClientData,
 			LobbyDetails->GetInfo(),
 			MoveTemp(UnregisterFn)));
 	});
@@ -851,7 +813,7 @@ TFuture<TDefaultErrorResultInternal<TSharedRef<FLobbyDataEOS>>> FLobbyDataRegist
 	const FLobbyId LobbyId = FLobbyId(EOnlineServices::Epic, NextHandleIndex++);
 
 	// Create lobby data from lobby details.
-	FLobbyDataEOS::Create(LobbyId, LobbyDetails, MakeUnregisterFn())
+	FLobbyDataEOS::Create(Prerequisites, LobbyId, LobbyDetails, MakeUnregisterFn())
 	.Next([WeakThis = AsWeak(), Promise = MoveTemp(Promise), LocalAccountId, LobbyDetails](TDefaultErrorResultInternal<TSharedRef<FLobbyDataEOS>>&& Result) mutable
 	{
 		TSharedPtr<FLobbyDataRegistryEOS> StrongThis = WeakThis.Pin();
@@ -1062,10 +1024,21 @@ TFuture<TDefaultErrorResultInternal<TSharedRef<FLobbySearchEOS>>> FLobbySearchEO
 			}
 		}
 
-		for (const FFindLobbySearchFilter& Filter :  Params.Filters)
+		// Create schema category instance to verify search filter attributes using the base lobby schema.
+		FSchemaCategoryInstance SchemaCategoryInstance(FSchemaId(), LobbyBaseSchemaId, LobbySchemaCategoryId, Prerequisites->SchemaRegistry);
+		for (const FFindLobbySearchFilter& Filter : Params.Filters)
 		{
+			FSchemaServiceAttributeId SchemaServiceAttributeId;
+			ESchemaServiceAttributeFlags SchemaServiceAttributeFlags = ESchemaServiceAttributeFlags::None;
+			if (!SchemaCategoryInstance.VerifyBaseAttributeData(Filter.AttributeName, Filter.ComparisonValue, SchemaServiceAttributeId, SchemaServiceAttributeFlags))
+			{
+				UE_LOG(LogTemp, Warning, TEXT("[FLobbySearchEOS::Create] VerifyBaseAttributeData failed: User[%s], Key[%s], Value[%s]"),
+					*ToLogString(Params.LocalAccountId), *Filter.AttributeName.ToString().ToLower(), *ToLogString(Filter.ComparisonValue));
+				return MakeFulfilledPromise<TDefaultErrorResultInternal<TSharedRef<FLobbySearchEOS>>>(Errors::InvalidParams()).GetFuture();
+			}
+
 			const FLobbyAttributeTranslator<ELobbyTranslationType::ToService> AttributeTranslator(
-				Filter.AttributeName, Filter.ComparisonValue);
+				{ SchemaServiceAttributeId, SchemaServiceAttributeFlags, Filter.ComparisonValue });
 
 			EOS_LobbySearch_SetParameterOptions SetParameterOptions = {};
 			SetParameterOptions.ApiVersion = EOS_LOBBYSEARCH_SETPARAMETER_API_LATEST;
@@ -1173,7 +1146,7 @@ TArray<TSharedRef<const FLobby>> FLobbySearchEOS::GetLobbyResults() const
 
 	for (const TSharedRef<FLobbyDataEOS>& LobbyData : Lobbies)
 	{
-		Result.Add(LobbyData->GetClientLobbyData()->GetPublicDataPtr());
+		Result.Add(LobbyData->GetLobbyClientData()->GetPublicDataPtr());
 	}
 
 	return Result;
diff --git a/Engine/Plugins/Online/Experimental/OnlineServicesEOSGS/Source/Private/Online/LobbiesEOSGSTypes.h b/Engine/Plugins/Online/Experimental/OnlineServicesEOSGS/Source/Private/Online/LobbiesEOSGSTypes.h
--- a/Engine/Plugins/Online/Experimental/OnlineServicesEOSGS/Source/Private/Online/LobbiesEOSGSTypes.h
+++ b/Engine/Plugins/Online/Experimental/OnlineServicesEOSGS/Source/Private/Online/LobbiesEOSGSTypes.h
@@ -23,7 +23,7 @@ class FLobbyInviteDataEOSRegistry;
 class FLobbySchema;
 struct FLobbyEvents;
 class FAuthEOSGS;
-struct FClientLobbyMemberSnapshot;
+struct FLobbyMemberServiceSnapshot;
 
 //--------------------------------------------------------------------------------------------------
 // Translated types
@@ -85,15 +85,16 @@ template <>
 class FLobbyAttributeTranslator<ELobbyTranslationType::ToService>
 {
 public:
-	FLobbyAttributeTranslator(const TPair<FLobbyAttributeId, FLobbyVariant>& FromAttributeData);
-	FLobbyAttributeTranslator(FLobbyAttributeId FromAttributeId, const FLobbyVariant& FromAttributeData);
+	FLobbyAttributeTranslator(const FSchemaServiceAttributeData& FromAttributeData);
 
 	const EOS_Lobby_AttributeData& GetAttributeData() const { return AttributeData; }
+	const EOS_ELobbyAttributeVisibility& GetAttributeVisibility() const { return AttributeVisibility; }
 
 private:
 	FTCHARToUTF8 KeyConverterStorage;
 	TOptional<FTCHARToUTF8> ValueConverterStorage;
 	EOS_Lobby_AttributeData AttributeData;
+	EOS_ELobbyAttributeVisibility AttributeVisibility;
 };
 
 template <>
@@ -102,11 +103,11 @@ class FLobbyAttributeTranslator<ELobbyTranslationType::FromService>
 public:
 	FLobbyAttributeTranslator(const EOS_Lobby_AttributeData& FromAttributeData);
 
-	const TPair<FLobbyAttributeId, FLobbyVariant>& GetAttributeData() const { return AttributeData; }
-	TPair<FLobbyAttributeId, FLobbyVariant>& GetMutableAttributeData() { return AttributeData; }
+	const TPair<FSchemaAttributeId, FSchemaVariant>& GetAttributeData() const { return AttributeData; }
+	TPair<FSchemaAttributeId, FSchemaVariant>& GetMutableAttributeData() { return AttributeData; }
 
 private:
-	TPair<FLobbyAttributeId, FLobbyVariant> AttributeData;
+	TPair<FSchemaAttributeId, FSchemaVariant> AttributeData;
 };
 
 //--------------------------------------------------------------------------------------------------
@@ -180,9 +181,8 @@ struct FLobbyPrerequisitesEOS
 {
 	EOS_HLobby LobbyInterfaceHandle = {};
 	TWeakPtr<FAuthEOSGS> AuthInterface;
-	TSharedRef<const FLobbySchemaRegistry> SchemaRegistry;
-	TSharedRef<const FLobbySchema> ServiceSchema;
 	FLobbyBucketIdEOS BucketId;
+	TSharedRef<FSchemaRegistry> SchemaRegistry;
 };
 
 /**
@@ -262,25 +262,18 @@ public:
 	 * Retrieve a lobby data snapshot from the EOS lobby details object.
 	 * The lobby schema will be used to translate attribute data before returning.
 	 */
-	TFuture<TDefaultErrorResultInternal<TSharedRef<FClientLobbySnapshot>>> GetLobbySnapshot() const;
+	TFuture<TDefaultErrorResultInternal<FLobbyServiceSnapshot>> GetLobbySnapshot() const;
 
 	/**
 	 * Retrieve lobby member data snapshot from the EOS lobby details object.
 	 * The lobby schema will be used to translate attribute data before returning.
 	 */
-	TDefaultErrorResultInternal<TSharedRef<FClientLobbyMemberSnapshot>> GetLobbyMemberSnapshot(FAccountId MemberAccountId) const;
+	TDefaultErrorResultInternal<FLobbyMemberServiceSnapshot> GetLobbyMemberSnapshot(FAccountId MemberAccountId) const;
 
 	/**
 	 * Apply client side lobby changes to the lobby service.
-	 * The lobby schema will be used to translate any changed attributes before sending.
 	 */
-	TFuture<EOS_EResult> ApplyLobbyDataUpdateFromLocalChanges(FAccountId LocalAccountId, const FClientLobbyDataChanges& Changes) const;
-
-	/**
-	 * Apply client side lobby member changes to the lobby service.
-	 * The lobby schema will be used to translate any changed attributes before sending.
-	 */
-	TFuture<EOS_EResult> ApplyLobbyMemberDataUpdateFromLocalChanges(FAccountId LocalAccountId, const FClientLobbyMemberDataChanges& Changes) const;
+	TFuture<EOS_EResult> ApplyLobbyDataUpdateFromLocalChanges(FAccountId LocalAccountId, const FLobbyClientServiceChanges& ServiceChanges) const;
 
 private:
 	template <typename, ESPMode>
@@ -342,14 +335,15 @@ public:
 	FLobbyDataEOS() = delete;
 
 	static TFuture<TDefaultErrorResultInternal<TSharedRef<FLobbyDataEOS>>> Create(
+		TSharedRef<FLobbyPrerequisitesEOS> Prerequisites,
 		FLobbyId LobbyId,
 		const TSharedRef<FLobbyDetailsEOS>& LobbyDetails,
 		FUnregisterFn UnregisterFn = FUnregisterFn());
 
 	~FLobbyDataEOS();
 
-	FLobbyId GetLobbyIdHandle() const { return ClientLobbyData->GetPublicData().LobbyId; }
-	const TSharedRef<FClientLobbyData>& GetClientLobbyData() const { return ClientLobbyData; }
+	FLobbyId GetLobbyIdHandle() const { return LobbyClientData->GetPublicData().LobbyId; }
+	const TSharedRef<FLobbyClientData>& GetLobbyClientData() const { return LobbyClientData; }
 	EOS_LobbyId GetLobbyIdEOS() const { return LobbyDetailsInfo->GetLobbyIdEOS(); }
 	const FString& GetLobbyIdString() const { return LobbyIdString; }
 
@@ -367,11 +361,11 @@ private:
 	friend class SharedPointerInternals::TIntrusiveReferenceController;
 
 	FLobbyDataEOS(
-		const TSharedRef<FClientLobbyData>& ClientLobbyData,
+		const TSharedRef<FLobbyClientData>& LobbyClientData,
 		const TSharedRef<FLobbyDetailsInfoEOS>& LobbyDetailsInfo,
 		FUnregisterFn UnregisterFn);
 
-	TSharedRef<FClientLobbyData> ClientLobbyData;
+	TSharedRef<FLobbyClientData> LobbyClientData;
 	TSharedRef<FLobbyDetailsInfoEOS> LobbyDetailsInfo;
 	FUnregisterFn UnregisterFn;
 	FString LobbyIdString;
diff --git a/Engine/Plugins/Online/Experimental/OnlineServicesEOSGS/Source/Public/Online/LobbiesEOSGS.h b/Engine/Plugins/Online/Experimental/OnlineServicesEOSGS/Source/Public/Online/LobbiesEOSGS.h
--- a/Engine/Plugins/Online/Experimental/OnlineServicesEOSGS/Source/Public/Online/LobbiesEOSGS.h
+++ b/Engine/Plugins/Online/Experimental/OnlineServicesEOSGS/Source/Public/Online/LobbiesEOSGS.h
@@ -151,37 +151,10 @@ struct FLobbiesModifyLobbyDataImpl
 		// The local user agent which will perform the action.
 		FAccountId LocalAccountId;
 
-		// The changes to apply to the lobby data.
-		TSharedPtr<FClientLobbyDataChanges> Changes;
+		/** Translated changes to be applied to the service. */
+		FLobbyClientServiceChanges ServiceChanges;
 	};
 
-	// Todo: += operator.
-	// Mergeable op must be queued by lobby id.
-
-	struct Result
-	{
-	};
-};
-
-struct FLobbiesModifyLobbyMemberDataImpl
-{
-	static constexpr TCHAR Name[] = TEXT("ModifyLobbyMemberDataImpl");
-
-	struct Params
-	{
-		// The lobby handle data.
-		TSharedPtr<FLobbyDataEOS> LobbyData;
-
-		// The local user agent which will perform the action.
-		FAccountId LocalAccountId;
-
-		// The changes to apply to the lobby member data.
-		TSharedPtr<FClientLobbyMemberDataChanges> Changes;
-	};
-
-	// Todo: += operator.
-	// Mergeable op must be queued on union of lobby id + lobby member account id.
-
 	struct Result
 	{
 	};
@@ -264,7 +237,6 @@ protected:
 	TFuture<TDefaultErrorResult<FLobbiesKickLobbyMemberImpl>> KickLobbyMemberImpl(FLobbiesKickLobbyMemberImpl::Params&& Params);
 	TFuture<TDefaultErrorResult<FLobbiesPromoteLobbyMemberImpl>> PromoteLobbyMemberImpl(FLobbiesPromoteLobbyMemberImpl::Params&& Params);
 	TFuture<TDefaultErrorResult<FLobbiesModifyLobbyDataImpl>> ModifyLobbyDataImpl(FLobbiesModifyLobbyDataImpl::Params&& Params);
-	TFuture<TDefaultErrorResult<FLobbiesModifyLobbyMemberDataImpl>> ModifyLobbyMemberDataImpl(FLobbiesModifyLobbyMemberDataImpl::Params&& Params);
 	TOnlineAsyncOpHandle<FLobbiesProcessLobbyNotificationImpl> ProcessLobbyNotificationImplOp(FLobbiesProcessLobbyNotificationImpl::Params&& Params);
 
 	EOSEventRegistrationPtr OnLobbyUpdatedEOSEventRegistration;
@@ -343,14 +315,6 @@ END_ONLINE_STRUCT_META()
 BEGIN_ONLINE_STRUCT_META(FLobbiesModifyLobbyDataImpl::Result)
 END_ONLINE_STRUCT_META()
 
-BEGIN_ONLINE_STRUCT_META(FLobbiesModifyLobbyMemberDataImpl::Params)
-	ONLINE_STRUCT_FIELD(FLobbiesModifyLobbyMemberDataImpl::Params, LobbyData),
-	ONLINE_STRUCT_FIELD(FLobbiesModifyLobbyMemberDataImpl::Params, LocalAccountId)
-END_ONLINE_STRUCT_META()
-
-BEGIN_ONLINE_STRUCT_META(FLobbiesModifyLobbyMemberDataImpl::Result)
-END_ONLINE_STRUCT_META()
-
 BEGIN_ONLINE_STRUCT_META(FLobbiesProcessLobbyNotificationImpl::Params)
 	ONLINE_STRUCT_FIELD(FLobbiesProcessLobbyNotificationImpl::Params, LobbyData),
 	ONLINE_STRUCT_FIELD(FLobbiesProcessLobbyNotificationImpl::Params, MutatedMembers),
diff --git a/Engine/Plugins/Online/Experimental/OnlineServicesNull/Source/Public/Online/NboSerializerNullSvc.h b/Engine/Plugins/Online/Experimental/OnlineServicesNull/Source/Public/Online/NboSerializerNullSvc.h
--- a/Engine/Plugins/Online/Experimental/OnlineServicesNull/Source/Public/Online/NboSerializerNullSvc.h
+++ b/Engine/Plugins/Online/Experimental/OnlineServicesNull/Source/Public/Online/NboSerializerNullSvc.h
@@ -43,10 +43,10 @@ inline void SerializeToBuffer(FNboSerializeToBuffer& Packet, const FSessionMembe
 	}
 }
 
-inline void SerializeToBuffer(FNboSerializeToBuffer& Ar, const TMap<FLobbyAttributeId, FLobbyVariant>& Map)
+inline void SerializeToBuffer(FNboSerializeToBuffer& Ar, const TMap<FSchemaAttributeId, FSchemaVariant>& Map)
 {
 	Ar << Map.Num();
-	for (const TPair<FLobbyAttributeId, FLobbyVariant>& Pair : Map)
+	for (const TPair<FSchemaAttributeId, FSchemaVariant>& Pair : Map)
 	{
 		Ar << Pair.Key;
 		Ar << Pair.Value.GetString();
@@ -89,13 +89,13 @@ inline void SerializeFromBuffer(FNboSerializeFromBuffer& Packet, FSessionMemberI
 	}
 }
 
-inline void SerializeFromBuffer(FNboSerializeFromBuffer& Ar, TMap<FLobbyAttributeId, FLobbyVariant>& Map)
+inline void SerializeFromBuffer(FNboSerializeFromBuffer& Ar, TMap<FSchemaAttributeId, FSchemaVariant>& Map)
 {
 	int32 Size;
 	Ar >> Size;
 	for (int i = 0; i < Size; i++)
 	{
-		FLobbyAttributeId LobbyAttributeId;
+		FSchemaAttributeId LobbyAttributeId;
 		FString LobbyData;
 		Ar >> LobbyAttributeId;
 		Ar >> LobbyData;
