Add lobby schema config for OnlineTests.
Add lobby schema config for Lyra.
Add explicit schema id verification when creating a lobby.

diff --git a/Engine/Plugins/Online/Experimental/OnlineServicesEOSGS/Source/Private/Online/LobbiesEOSGS.cpp b/Engine/Plugins/Online/Experimental/OnlineServicesEOSGS/Source/Private/Online/LobbiesEOSGS.cpp
--- a/Engine/Plugins/Online/Experimental/OnlineServicesEOSGS/Source/Private/Online/LobbiesEOSGS.cpp
+++ b/Engine/Plugins/Online/Experimental/OnlineServicesEOSGS/Source/Private/Online/LobbiesEOSGS.cpp
@@ -107,6 +107,14 @@ TOnlineAsyncOpHandle<FCreateLobby> FLobbiesEOSGS::CreateLobby(FCreateLobby::Para
 		return Op->GetHandle();
 	}
 
+	// Check that requested lobby schema exists.
+	if (!SchemaRegistry->GetDefinition(Params.SchemaId).IsValid())
+	{
+		UE_LOG(LogTemp, Warning, TEXT("[FLobbiesEOSGS::CreateLobby] Failed: Unknown lobby schema [%s]"), *Params.SchemaId.ToString().ToLower());
+		Op->SetError(Errors::InvalidParams());
+		return Op->GetHandle();
+	}
+
 	// Start operation.
 	// Step 1: Call create lobby.
 	Op->Then([this](TOnlineAsyncOp<FCreateLobby>& InAsyncOp, TPromise<const EOS_Lobby_CreateLobbyCallbackInfo*>&& Promise)
