- Convert AchievementsEOSGS to joinable ops
- Move state checking inside .Then

diff --git a/Engine/Plugins/Online/Experimental/OnlineServicesEOSGS/Source/Private/Online/AchievementsEOSGS.cpp b/Engine/Plugins/Online/Experimental/OnlineServicesEOSGS/Source/Private/Online/AchievementsEOSGS.cpp
--- a/Engine/Plugins/Online/Experimental/OnlineServicesEOSGS/Source/Private/Online/AchievementsEOSGS.cpp
+++ b/Engine/Plugins/Online/Experimental/OnlineServicesEOSGS/Source/Private/Online/AchievementsEOSGS.cpp
@@ -72,22 +72,30 @@ void FAchievementsEOSGS::Shutdown()
 
 TOnlineAsyncOpHandle<FQueryAchievementDefinitions> FAchievementsEOSGS::QueryAchievementDefinitions(FQueryAchievementDefinitions::Params&& Params)
 {
-	TOnlineAsyncOpRef<FQueryAchievementDefinitions> Op = GetOp<FQueryAchievementDefinitions>(MoveTemp(Params));
-
-	if (!Services.Get<FAuthEOSGS>()->IsLoggedIn(Op->GetParams().LocalAccountId))
+	TOnlineAsyncOpRef<FQueryAchievementDefinitions> Op = GetJoinableOp<FQueryAchievementDefinitions>(MoveTemp(Params));
+	if (!Op->IsReady())
 	{
-		Op->SetError(Errors::InvalidUser());
-		return Op->GetHandle();
-	}
+		Op->Then([this](TOnlineAsyncOp<FQueryAchievementDefinitions>& InAsyncOp)
+		{
+			const FQueryAchievementDefinitions::Params& Params = InAsyncOp.GetParams();
 
-	if (!AchievementDefinitions.IsSet())
-	{
-		Op->Then([this](TOnlineAsyncOp<FQueryAchievementDefinitions>& InAsyncOp, TPromise<const EOS_Achievements_OnQueryDefinitionsCompleteCallbackInfo*>&& Promise)
+			if (!Services.Get<FAuthEOSGS>()->IsLoggedIn(Params.LocalAccountId))
+			{
+				InAsyncOp.SetError(Errors::InvalidUser());
+			}
+			else if (AchievementDefinitions.IsSet())
+			{
+				InAsyncOp.SetResult({});
+			}
+		})
+		.Then([this](TOnlineAsyncOp<FQueryAchievementDefinitions>& InAsyncOp, TPromise<const EOS_Achievements_OnQueryDefinitionsCompleteCallbackInfo*>&& Promise)
 		{
+			const FQueryAchievementDefinitions::Params& Params = InAsyncOp.GetParams();
+
 			EOS_Achievements_QueryDefinitionsOptions Options = {};
 			Options.ApiVersion = EOS_ACHIEVEMENTS_QUERYDEFINITIONS_API_LATEST;
 			static_assert(EOS_ACHIEVEMENTS_QUERYDEFINITIONS_API_LATEST == 3, "EOS_Achievements_QueryDefinitionsOptions updated, check new fields");
-			Options.LocalUserId = GetProductUserIdChecked(InAsyncOp.GetParams().LocalAccountId);
+			Options.LocalUserId = GetProductUserIdChecked(Params.LocalAccountId);
 
 			EOS_Async(EOS_Achievements_QueryDefinitions, AchievementsHandle, Options, MoveTemp(Promise));
 		})
@@ -160,10 +168,6 @@ TOnlineAsyncOpHandle<FQueryAchievementDefinitions> FAchievementsEOSGS::QueryAchi
 		})
 		.Enqueue(GetSerialQueue());
 	}
-	else
-	{
-		Op->SetResult({});
-	}
 	return Op->GetHandle();
 }
 
@@ -209,35 +213,42 @@ TOnlineResult<FGetAchievementDefinition> FAchievementsEOSGS::GetAchievementDefin
 
 TOnlineAsyncOpHandle<FQueryAchievementStates> FAchievementsEOSGS::QueryAchievementStates(FQueryAchievementStates::Params&& Params)
 {
-	TOnlineAsyncOpRef<FQueryAchievementStates> Op = GetOp<FQueryAchievementStates>(MoveTemp(Params));
-
-	if (!Services.Get<FAuthEOSGS>()->IsLoggedIn(Op->GetParams().LocalAccountId))
-	{
-		Op->SetError(Errors::InvalidUser());
-		return Op->GetHandle();
-	}
-
-	if (!AchievementDefinitions.IsSet())
+	TOnlineAsyncOpRef<FQueryAchievementStates> Op = GetJoinableOp<FQueryAchievementStates>(MoveTemp(Params));
+	if (!Op->IsReady())
 	{
-		// Call QueryAchievementDefinitions first
-		Op->SetError(Errors::InvalidState());
-		return Op->GetHandle();
-	}
-
-	if (!AchievementStates.Contains(Op->GetParams().LocalAccountId))
-	{
-		Op->Then([this](TOnlineAsyncOp<FQueryAchievementStates>& InAsyncOp, TPromise<const EOS_Achievements_OnQueryPlayerAchievementsCompleteCallbackInfo*>&& Promise)
+		Op->Then([this](TOnlineAsyncOp<FQueryAchievementStates>& InAsyncOp)
 		{
+			const FQueryAchievementStates::Params& Params = InAsyncOp.GetParams();
+			if (!Services.Get<FAuthEOSGS>()->IsLoggedIn(Params.LocalAccountId))
+			{
+				InAsyncOp.SetError(Errors::InvalidUser());
+			}
+			else if (!AchievementDefinitions.IsSet())
+			{
+				// Call QueryAchievementDefinitions first
+				InAsyncOp.SetError(Errors::InvalidState());
+			}
+			else if (AchievementStates.Contains(Params.LocalAccountId))
+			{
+				InAsyncOp.SetResult({});
+			}
+		})
+		.Then([this](TOnlineAsyncOp<FQueryAchievementStates>& InAsyncOp, TPromise<const EOS_Achievements_OnQueryPlayerAchievementsCompleteCallbackInfo*>&& Promise)
+		{
+			const FQueryAchievementStates::Params& Params = InAsyncOp.GetParams();
+
 			EOS_Achievements_QueryPlayerAchievementsOptions Options = {};
 			Options.ApiVersion = EOS_ACHIEVEMENTS_QUERYPLAYERACHIEVEMENTS_API_LATEST;
 			static_assert(EOS_ACHIEVEMENTS_QUERYPLAYERACHIEVEMENTS_API_LATEST == 2, "EOS_Achievements_QueryPlayerAchievementsOptions updated, check new fields");
-			Options.LocalUserId = GetProductUserIdChecked(InAsyncOp.GetParams().LocalAccountId);
+			Options.LocalUserId = GetProductUserIdChecked(Params.LocalAccountId);
 			Options.TargetUserId = Options.LocalUserId;
 
 			EOS_Async(EOS_Achievements_QueryPlayerAchievements, AchievementsHandle, Options, MoveTemp(Promise));
 		})
 		.Then([this](TOnlineAsyncOp<FQueryAchievementStates>& InAsyncOp, const EOS_Achievements_OnQueryPlayerAchievementsCompleteCallbackInfo* Data)
 		{
+			const FQueryAchievementStates::Params& Params = InAsyncOp.GetParams();
+
 			if (Data->ResultCode != EOS_EResult::EOS_Success)
 			{
 				UE_LOG(LogTemp, Warning, TEXT("EOS_Achievements_QueryPlayerAchievements failed with result=[%s]"), *LexToString(Data->ResultCode));
@@ -249,7 +260,7 @@ TOnlineAsyncOpHandle<FQueryAchievementStates> FAchievementsEOSGS::QueryAchieveme
 			EOS_Achievements_GetPlayerAchievementCountOptions GetCountOptions = {};
 			GetCountOptions.ApiVersion = EOS_ACHIEVEMENTS_GETPLAYERACHIEVEMENTCOUNT_API_LATEST;
 			static_assert(EOS_ACHIEVEMENTS_GETPLAYERACHIEVEMENTCOUNT_API_LATEST == 1, "EOS_Achievements_GetPlayerAchievementCountOptions updated, check new fields");
-			GetCountOptions.UserId = GetProductUserIdChecked(InAsyncOp.GetParams().LocalAccountId);
+			GetCountOptions.UserId = GetProductUserIdChecked(Params.LocalAccountId);
 
 			const uint32_t AchievementCount = EOS_Achievements_GetPlayerAchievementCount(AchievementsHandle, &GetCountOptions);
 
@@ -259,7 +270,7 @@ TOnlineAsyncOpHandle<FQueryAchievementStates> FAchievementsEOSGS::QueryAchieveme
 				EOS_Achievements_CopyPlayerAchievementByIndexOptions CopyOptions = {};
 				CopyOptions.ApiVersion = EOS_ACHIEVEMENTS_COPYPLAYERACHIEVEMENTBYINDEX_API_LATEST;
 				static_assert(EOS_ACHIEVEMENTS_COPYPLAYERACHIEVEMENTBYINDEX_API_LATEST == 2, "EOS_Achievements_CopyPlayerAchievementByIndexOptions updated, check new fields");
-				CopyOptions.LocalUserId = GetProductUserIdChecked(InAsyncOp.GetParams().LocalAccountId);
+				CopyOptions.LocalUserId = GetProductUserIdChecked(Params.LocalAccountId);
 				CopyOptions.TargetUserId = CopyOptions.LocalUserId;
 				CopyOptions.AchievementIndex = AchievementIdx;
 
@@ -272,7 +283,7 @@ TOnlineAsyncOpHandle<FQueryAchievementStates> FAchievementsEOSGS::QueryAchieveme
 					InAsyncOp.SetError(Errors::FromEOSResult(CopyResult));
 					return;
 				}
-				
+
 				static_assert(EOS_ACHIEVEMENTS_PLAYERACHIEVEMENT_API_LATEST == 2, "EOS_Achievements_PlayerAchievement updated, check new fields");
 				if (!ensure(EOSPlayerAchievement->ApiVersion == EOS_ACHIEVEMENTS_PLAYERACHIEVEMENT_API_LATEST))
 				{
@@ -287,16 +298,11 @@ TOnlineAsyncOpHandle<FQueryAchievementStates> FAchievementsEOSGS::QueryAchieveme
 
 				EOS_Achievements_PlayerAchievement_Release(EOSPlayerAchievement);
 			}
-			AchievementStates.Emplace(InAsyncOp.GetParams().LocalAccountId, MoveTemp(NewAchievementStates));
+			AchievementStates.Emplace(Params.LocalAccountId, MoveTemp(NewAchievementStates));
 			InAsyncOp.SetResult({});
 		})
 		.Enqueue(GetSerialQueue());
 	}
-	else
-	{
-		Op->SetResult({});
-	}
-
 	return Op->GetHandle();
 }
 
@@ -325,70 +331,80 @@ TOnlineResult<FGetAchievementState> FAchievementsEOSGS::GetAchievementState(FGet
 
 TOnlineAsyncOpHandle<FUnlockAchievements> FAchievementsEOSGS::UnlockAchievements(FUnlockAchievements::Params&& Params)
 {
-	TOnlineAsyncOpRef<FUnlockAchievements> Op = GetOp<FUnlockAchievements>(MoveTemp(Params));
-
-	if (!Services.Get<FAuthEOSGS>()->IsLoggedIn(Op->GetParams().LocalAccountId))
+	TOnlineAsyncOpRef<FUnlockAchievements> Op = GetJoinableOp<FUnlockAchievements>(MoveTemp(Params));
+	if (!Op->IsReady())
 	{
-		Op->SetError(Errors::InvalidUser());
-		return Op->GetHandle();
-	}
+		Op->Then([this](TOnlineAsyncOp<FUnlockAchievements>& InAsyncOp)
+		{
+			const FUnlockAchievements::Params& Params = InAsyncOp.GetParams();
+			if (!Services.Get<FAuthEOSGS>()->IsLoggedIn(Params.LocalAccountId))
+			{
+				InAsyncOp.SetError(Errors::InvalidUser());
+				return;
+			}
 
-	if (Op->GetParams().AchievementIds.IsEmpty())
-	{
-		Op->SetError(Errors::InvalidParams());
-		return Op->GetHandle();
-	}
+			if (Params.AchievementIds.IsEmpty())
+			{
+				InAsyncOp.SetError(Errors::InvalidParams());
+				return;
+			}
 
-	FAchievementStateMap* LocalUserAchievementStates = AchievementStates.Find(Op->GetParams().LocalAccountId);
-	if (!LocalUserAchievementStates)
-	{
-		// Call QueryAchievementStates first
-		Op->SetError(Errors::InvalidState());
-		return Op->GetHandle();
-	}
+			FAchievementStateMap* LocalUserAchievementStates = AchievementStates.Find(Params.LocalAccountId);
+			if (!LocalUserAchievementStates)
+			{
+				// Call QueryAchievementStates first
+				InAsyncOp.SetError(Errors::InvalidState());
+				return;
+			}
 
-	for (const FString& AchievementId : Op->GetParams().AchievementIds)
-	{
-		const FAchievementState* AchievementState = LocalUserAchievementStates->Find(AchievementId);
-		if (!AchievementState)
-		{
-			Op->SetError(Errors::NotFound());
-			return Op->GetHandle();
-		}
-		if (AchievementState->Progress == 1.0f)
+			for (const FString& AchievementId : Params.AchievementIds)
+			{
+				const FAchievementState* AchievementState = LocalUserAchievementStates->Find(AchievementId);
+				if (!AchievementState)
+				{
+					InAsyncOp.SetError(Errors::NotFound());
+					return;
+				}
+				if (AchievementState->Progress == 1.0f)
+				{
+					InAsyncOp.SetError(Errors::Achievements::AlreadyUnlocked());
+					return;
+				}
+			}
+		})
+		.Then([this](TOnlineAsyncOp<FUnlockAchievements>& InAsyncOp, TPromise<const EOS_Achievements_OnUnlockAchievementsCompleteCallbackInfo*>&& Promise)
 		{
-			Op->SetError(Errors::Achievements::AlreadyUnlocked());
-			return Op->GetHandle();
-		}
-	}
-
-	Op->Then([this](TOnlineAsyncOp<FUnlockAchievements>& InAsyncOp, TPromise<const EOS_Achievements_OnUnlockAchievementsCompleteCallbackInfo*>&& Promise)
-	{
-		const TArray<FString>& InAchievementIds = InAsyncOp.GetParams().AchievementIds;
+			const FUnlockAchievements::Params& Params = InAsyncOp.GetParams();
 
-		TArray<FTCHARToUTF8> AchievementIdConverters;
-		AchievementIdConverters.Reserve(InAchievementIds.Num());
-		TArray<const char*> AchievementIdPtrs;
-		AchievementIdPtrs.Reserve(InAchievementIds.Num());
-		for (const FString& AchievementId : InAchievementIds)
-		{
-			const FTCHARToUTF8& Converter = AchievementIdConverters.Emplace_GetRef(*AchievementId);
-			AchievementIdPtrs.Emplace(Converter.Get());
-		}
+			const TArray<FString>& InAchievementIds = Params.AchievementIds;
+			TArray<FTCHARToUTF8> AchievementIdConverters;
+			AchievementIdConverters.Reserve(InAchievementIds.Num());
+			TArray<const char*> AchievementIdPtrs;
+			AchievementIdPtrs.Reserve(InAchievementIds.Num());
+			for (const FString& AchievementId : InAchievementIds)
+			{
+				const FTCHARToUTF8& Converter = AchievementIdConverters.Emplace_GetRef(*AchievementId);
+				AchievementIdPtrs.Emplace(Converter.Get());
+			}
 
-		EOS_Achievements_UnlockAchievementsOptions Options = {};
-		Options.ApiVersion = EOS_ACHIEVEMENTS_UNLOCKACHIEVEMENTS_API_LATEST;
-		static_assert(EOS_ACHIEVEMENTS_UNLOCKACHIEVEMENTS_API_LATEST == 1, "EOS_Achievements_UnlockAchievementsOptions updated, check new fields");
-		Options.UserId = GetProductUserIdChecked(InAsyncOp.GetParams().LocalAccountId);
-		Options.AchievementIds = AchievementIdPtrs.GetData();
-		Options.AchievementsCount = AchievementIdPtrs.Num();
+			EOS_Achievements_UnlockAchievementsOptions Options = {};
+			Options.ApiVersion = EOS_ACHIEVEMENTS_UNLOCKACHIEVEMENTS_API_LATEST;
+			static_assert(EOS_ACHIEVEMENTS_UNLOCKACHIEVEMENTS_API_LATEST == 1, "EOS_Achievements_UnlockAchievementsOptions updated, check new fields");
+			Options.UserId = GetProductUserIdChecked(Params.LocalAccountId);
+			Options.AchievementIds = AchievementIdPtrs.GetData();
+			Options.AchievementsCount = AchievementIdPtrs.Num();
 
-		EOS_Async(EOS_Achievements_UnlockAchievements, AchievementsHandle, Options, MoveTemp(Promise));
-	})
-	.Then([this](TOnlineAsyncOp<FUnlockAchievements>& InAsyncOp, const EOS_Achievements_OnUnlockAchievementsCompleteCallbackInfo* Data)
-	{
-		if (Data->ResultCode == EOS_EResult::EOS_Success)
+			EOS_Async(EOS_Achievements_UnlockAchievements, AchievementsHandle, Options, MoveTemp(Promise));
+		})
+		.Then([this](TOnlineAsyncOp<FUnlockAchievements>& InAsyncOp, const EOS_Achievements_OnUnlockAchievementsCompleteCallbackInfo* Data)
 		{
+			if (Data->ResultCode != EOS_EResult::EOS_Success)
+			{
+				UE_LOG(LogTemp, Warning, TEXT("EOS_Achievements_UnlockAchievements failed with result=[%s]"), *LexToString(Data->ResultCode));
+				InAsyncOp.SetError(Errors::FromEOSResult(Data->ResultCode));
+				return;
+			}
+
 			UE_LOG(LogTemp, Verbose, TEXT("EOS_Achievements_UnlockAchievements succeeded"), *LexToString(Data->ResultCode));
 
 			const FUnlockAchievements::Params& Params = InAsyncOp.GetParams();
@@ -397,7 +413,7 @@ TOnlineAsyncOpHandle<FUnlockAchievements> FAchievementsEOSGS::UnlockAchievements
 			FAchievementStateMap* LocalUserAchievementStates = AchievementStates.Find(Params.LocalAccountId);
 			if (ensure(LocalUserAchievementStates))
 			{
-				for(const FString& AchievementId : Params.AchievementIds)
+				for (const FString& AchievementId : Params.AchievementIds)
 				{
 					FAchievementState* AchievementState = LocalUserAchievementStates->Find(AchievementId);
 					if (ensure(AchievementState))
@@ -407,17 +423,10 @@ TOnlineAsyncOpHandle<FUnlockAchievements> FAchievementsEOSGS::UnlockAchievements
 					}
 				}
 			}
-
 			InAsyncOp.SetResult({});
-		}
-		else
-		{
-			UE_LOG(LogTemp, Warning, TEXT("EOS_Achievements_UnlockAchievements failed with result=[%s]"), *LexToString(Data->ResultCode));
-			InAsyncOp.SetError(Errors::FromEOSResult(Data->ResultCode));
-		}
-	})
-	.Enqueue(GetSerialQueue());
-
+		})
+		.Enqueue(GetSerialQueue());
+	}
 	return Op->GetHandle();
 }
 
