- Use joinable ops in EOS TitleFile

diff --git a/Engine/Plugins/Online/Experimental/OnlineServicesEOSGS/Source/Private/Online/TitleFileEOSGS.cpp b/Engine/Plugins/Online/Experimental/OnlineServicesEOSGS/Source/Private/Online/TitleFileEOSGS.cpp
--- a/Engine/Plugins/Online/Experimental/OnlineServicesEOSGS/Source/Private/Online/TitleFileEOSGS.cpp
+++ b/Engine/Plugins/Online/Experimental/OnlineServicesEOSGS/Source/Private/Online/TitleFileEOSGS.cpp
@@ -35,88 +35,88 @@ void FTitleFileEOSGS::UpdateConfig()
 
 TOnlineAsyncOpHandle<FTitleFileEnumerateFiles> FTitleFileEOSGS::EnumerateFiles(FTitleFileEnumerateFiles::Params&& InParams)
 {
-	TOnlineAsyncOpRef<FTitleFileEnumerateFiles> Op = GetOp<FTitleFileEnumerateFiles>(MoveTemp(InParams));
-
-	Op->Then([this](TOnlineAsyncOp<FTitleFileEnumerateFiles>& InAsyncOp, TPromise<const EOS_TitleStorage_QueryFileListCallbackInfo*>&& Promise)
+	TOnlineAsyncOpRef<FTitleFileEnumerateFiles> Op = GetJoinableOp<FTitleFileEnumerateFiles>(MoveTemp(InParams));
+	if (!Op->IsReady())
 	{
-		const FTitleFileEnumerateFiles::Params& Params = InAsyncOp.GetParams();
-
-		if (!Services.Get<FAuthEOSGS>()->IsLoggedIn(Params.LocalAccountId))
+		Op->Then([this](TOnlineAsyncOp<FTitleFileEnumerateFiles>& InAsyncOp)
 		{
-			InAsyncOp.SetError(Errors::InvalidUser());
-			Promise.EmplaceValue();
-			return;
-		}
+			const FTitleFileEnumerateFiles::Params& Params = InAsyncOp.GetParams();
 
-		if (Config.SearchTag.IsEmpty())
+			if (!Services.Get<FAuthEOSGS>()->IsLoggedIn(Params.LocalAccountId))
+			{
+				InAsyncOp.SetError(Errors::InvalidUser());
+			}
+			else if (Config.SearchTag.IsEmpty())
+			{
+				InAsyncOp.SetError(Errors::NotConfigured());
+			}
+		})
+		.Then([this](TOnlineAsyncOp<FTitleFileEnumerateFiles>& InAsyncOp, TPromise<const EOS_TitleStorage_QueryFileListCallbackInfo*>&& Promise)
 		{
-			InAsyncOp.SetError(Errors::NotConfigured());
-			Promise.EmplaceValue();
-			return;
-		}
+			const FTitleFileEnumerateFiles::Params& Params = InAsyncOp.GetParams();
 
-		const FTCHARToUTF8 Utf8SearchTag(*Config.SearchTag);
-		const char* SearchTagPtr = Utf8SearchTag.Get();
+			const FTCHARToUTF8 Utf8SearchTag(*Config.SearchTag);
+			const char* SearchTagPtr = Utf8SearchTag.Get();
 
-		EOS_TitleStorage_QueryFileListOptions Options = {};
-		Options.ApiVersion = EOS_TITLESTORAGE_QUERYFILELISTOPTIONS_API_LATEST;
-		static_assert(EOS_TITLESTORAGE_QUERYFILELISTOPTIONS_API_LATEST == 1, "EOS_TitleStorage_QueryFileListOptions updated, check new fields");
-		Options.LocalUserId = GetProductUserIdChecked(Params.LocalAccountId);
-		Options.ListOfTagsCount = 1;
-		Options.ListOfTags = &SearchTagPtr;
+			EOS_TitleStorage_QueryFileListOptions Options = {};
+			Options.ApiVersion = EOS_TITLESTORAGE_QUERYFILELISTOPTIONS_API_LATEST;
+			static_assert(EOS_TITLESTORAGE_QUERYFILELISTOPTIONS_API_LATEST == 1, "EOS_TitleStorage_QueryFileListOptions updated, check new fields");
+			Options.LocalUserId = GetProductUserIdChecked(Params.LocalAccountId);
+			Options.ListOfTagsCount = 1;
+			Options.ListOfTags = &SearchTagPtr;
 
-		EOS_Async(EOS_TitleStorage_QueryFileList, TitleStorageHandle, Options, MoveTemp(Promise));
-	})
-	.Then([this](TOnlineAsyncOp<FTitleFileEnumerateFiles>& InAsyncOp, const EOS_TitleStorage_QueryFileListCallbackInfo* Data)
-	{
-		const FTitleFileEnumerateFiles::Params& Params = InAsyncOp.GetParams(); 
-		
-		if (Data->ResultCode != EOS_EResult::EOS_Success)
+			EOS_Async(EOS_TitleStorage_QueryFileList, TitleStorageHandle, Options, MoveTemp(Promise));
+		})
+		.Then([this](TOnlineAsyncOp<FTitleFileEnumerateFiles>& InAsyncOp, const EOS_TitleStorage_QueryFileListCallbackInfo* Data)
 		{
-			UE_LOG(LogTemp, Warning, TEXT("EOS_TitleStorage_QueryFileList failed with result=[%s]"), *LexToString(Data->ResultCode));
-			InAsyncOp.SetError(Errors::FromEOSResult(Data->ResultCode));
-			return;
-		}
+			if (Data->ResultCode != EOS_EResult::EOS_Success)
+			{
+				UE_LOG(LogTemp, Warning, TEXT("EOS_TitleStorage_QueryFileList failed with result=[%s]"), *LexToString(Data->ResultCode));
+				InAsyncOp.SetError(Errors::FromEOSResult(Data->ResultCode));
+				return;
+			}
 
-		const EOS_ProductUserId LocalUserPuid = GetProductUserIdChecked(Params.LocalAccountId);
+			const FTitleFileEnumerateFiles::Params& Params = InAsyncOp.GetParams();
+			const EOS_ProductUserId LocalUserPuid = GetProductUserIdChecked(Params.LocalAccountId);
 
-		EOS_TitleStorage_GetFileMetadataCountOptions GetCountOptions;
-		GetCountOptions.ApiVersion = EOS_TITLESTORAGE_GETFILEMETADATACOUNTOPTIONS_API_LATEST;
-		static_assert(EOS_TITLESTORAGE_GETFILEMETADATACOUNTOPTIONS_API_LATEST == 1, "EOS_TitleStorage_GetFileMetadataCountOptions updated, check new fields");
-		GetCountOptions.LocalUserId = LocalUserPuid;
+			EOS_TitleStorage_GetFileMetadataCountOptions GetCountOptions;
+			GetCountOptions.ApiVersion = EOS_TITLESTORAGE_GETFILEMETADATACOUNTOPTIONS_API_LATEST;
+			static_assert(EOS_TITLESTORAGE_GETFILEMETADATACOUNTOPTIONS_API_LATEST == 1, "EOS_TitleStorage_GetFileMetadataCountOptions updated, check new fields");
+			GetCountOptions.LocalUserId = LocalUserPuid;
 
-		const uint32_t NumFiles = EOS_TitleStorage_GetFileMetadataCount(TitleStorageHandle, &GetCountOptions);
+			const uint32_t NumFiles = EOS_TitleStorage_GetFileMetadataCount(TitleStorageHandle, &GetCountOptions);
 
-		TArray<FString> NewEnumeratedFiles;
-		NewEnumeratedFiles.Reserve(NumFiles);
+			TArray<FString> NewEnumeratedFiles;
+			NewEnumeratedFiles.Reserve(NumFiles);
 
-		for (uint32_t FileIdx = 0; FileIdx < NumFiles; FileIdx++)
-		{
-			EOS_TitleStorage_CopyFileMetadataAtIndexOptions CopyOptions;
-			CopyOptions.ApiVersion = EOS_TITLESTORAGE_COPYFILEMETADATAATINDEXOPTIONS_API_LATEST;
-			static_assert(EOS_TITLESTORAGE_COPYFILEMETADATAATINDEXOPTIONS_API_LATEST == 1, "EOS_TitleStorage_CopyFileMetadataAtIndexOptions updated, check new fields");
-			CopyOptions.LocalUserId = LocalUserPuid;
-			CopyOptions.Index = FileIdx;
+			for (uint32_t FileIdx = 0; FileIdx < NumFiles; FileIdx++)
+			{
+				EOS_TitleStorage_CopyFileMetadataAtIndexOptions CopyOptions;
+				CopyOptions.ApiVersion = EOS_TITLESTORAGE_COPYFILEMETADATAATINDEXOPTIONS_API_LATEST;
+				static_assert(EOS_TITLESTORAGE_COPYFILEMETADATAATINDEXOPTIONS_API_LATEST == 1, "EOS_TitleStorage_CopyFileMetadataAtIndexOptions updated, check new fields");
+				CopyOptions.LocalUserId = LocalUserPuid;
+				CopyOptions.Index = FileIdx;
 
-			EOS_TitleStorage_FileMetadata* FileMetadata = nullptr;
+				EOS_TitleStorage_FileMetadata* FileMetadata = nullptr;
 
-			const EOS_EResult Result = EOS_TitleStorage_CopyFileMetadataAtIndex(TitleStorageHandle, &CopyOptions, &FileMetadata);
-			if (Result != EOS_EResult::EOS_Success)
-			{
-				UE_LOG(LogTemp, Warning, TEXT("EOS_TitleStorage_CopyFileMetadataAtIndex failed with result=[%s]"), *LexToString(Result));
-				InAsyncOp.SetError(Errors::FromEOSResult(Result));
-				return;
-			}
+				const EOS_EResult Result = EOS_TitleStorage_CopyFileMetadataAtIndex(TitleStorageHandle, &CopyOptions, &FileMetadata);
+				if (Result != EOS_EResult::EOS_Success)
+				{
+					UE_LOG(LogTemp, Warning, TEXT("EOS_TitleStorage_CopyFileMetadataAtIndex failed with result=[%s]"), *LexToString(Result));
+					InAsyncOp.SetError(Errors::FromEOSResult(Result));
+					return;
+				}
 
-			NewEnumeratedFiles.Emplace(UTF8_TO_TCHAR(FileMetadata->Filename));
+				NewEnumeratedFiles.Emplace(UTF8_TO_TCHAR(FileMetadata->Filename));
 
-			EOS_TitleStorage_FileMetadata_Release(FileMetadata);
-		}
+				EOS_TitleStorage_FileMetadata_Release(FileMetadata);
+			}
 
-		EnumeratedFiles.Emplace(MoveTemp(NewEnumeratedFiles));
-		InAsyncOp.SetResult({});
-	})
-	.Enqueue(GetSerialQueue());
+			EnumeratedFiles.Emplace(MoveTemp(NewEnumeratedFiles));
+			InAsyncOp.SetResult({});
+		})
+		.Enqueue(GetSerialQueue());
+	}
 	return Op->GetHandle();
 }
 
@@ -154,71 +154,74 @@ TOnlineAsyncOpHandle<FTitleFileReadFile> FTitleFileEOSGS::ReadFile(FTitleFileRea
 	static const TCHAR* FileContentsKey = TEXT("FileContents");
 	static const TCHAR* RequestHandleKey = TEXT("RequestHandle");
 
-	TOnlineAsyncOpRef<FTitleFileReadFile> Op = GetOp<FTitleFileReadFile>(MoveTemp(InParams));
-	const FTitleFileReadFile::Params& Params = Op->GetParams();
-
-	if (Params.Filename.IsEmpty())
-	{
-		Op->SetError(Errors::InvalidParams());
-		return Op->GetHandle();
-	}
-
-	Op->Then([this](TOnlineAsyncOp<FTitleFileReadFile>& InAsyncOp, TPromise<const EOS_TitleStorage_ReadFileCallbackInfo*>&& Promise)
+	TOnlineAsyncOpRef<FTitleFileReadFile> Op = GetJoinableOp<FTitleFileReadFile>(MoveTemp(InParams));
+	if (!Op->IsReady())
 	{
-		const FTitleFileReadFile::Params& Params = InAsyncOp.GetParams();
-
-		if (!Services.Get<FAuthEOSGS>()->IsLoggedIn(Params.LocalAccountId))
+		const FTitleFileReadFile::Params& Params = Op->GetParams();
+		if (Params.Filename.IsEmpty())
 		{
-			InAsyncOp.SetError(Errors::InvalidUser());
-			Promise.EmplaceValue();
-			return;
+			Op->SetError(Errors::InvalidParams());
+			return Op->GetHandle();
 		}
 
-		FTitleFileContentsRef FileContentsRef = MakeShared<FTitleFileContents>();
-		InAsyncOp.Data.Set<FTitleFileContentsRef>(FileContentsKey, FileContentsRef);
+		Op->Then([this](TOnlineAsyncOp<FTitleFileReadFile>& InAsyncOp)
+		{
+			const FTitleFileReadFile::Params& Params = InAsyncOp.GetParams();
+			if (!Services.Get<FAuthEOSGS>()->IsLoggedIn(Params.LocalAccountId))
+			{
+				InAsyncOp.SetError(Errors::InvalidUser());
+			}
+		})
+		.Then([this](TOnlineAsyncOp<FTitleFileReadFile>& InAsyncOp, TPromise<const EOS_TitleStorage_ReadFileCallbackInfo*>&& Promise)
+		{
+			const FTitleFileReadFile::Params& Params = InAsyncOp.GetParams();
+
+			FTitleFileContentsRef FileContentsRef = MakeShared<FTitleFileContents>();
+			InAsyncOp.Data.Set<FTitleFileContentsRef>(FileContentsKey, FileContentsRef);
 
-		const FTCHARToUTF8 Utf8Filename(*Params.Filename);
+			const FTCHARToUTF8 Utf8Filename(*Params.Filename);
 
-		EOS_TitleStorage_ReadFileOptions Options = {};
-		Options.ApiVersion = EOS_TITLESTORAGE_READFILEOPTIONS_API_LATEST;
-		static_assert(EOS_TITLESTORAGE_READFILEOPTIONS_API_LATEST == 1, "EOS_TitleStorage_ReadFileOptions updated, check new fields");
-		Options.LocalUserId = GetProductUserIdChecked(Params.LocalAccountId);
-		Options.Filename = Utf8Filename.Get();
-		Options.ReadChunkLengthBytes = Config.ReadChunkLengthBytes;
-		Options.ReadFileDataCallback = &FTitleFileEOSGS::OnReadFileDataStatic;
-		Options.FileTransferProgressCallback = &FTitleFileEOSGS::OnFileTransferProgressStatic;
+			EOS_TitleStorage_ReadFileOptions Options = {};
+			Options.ApiVersion = EOS_TITLESTORAGE_READFILEOPTIONS_API_LATEST;
+			static_assert(EOS_TITLESTORAGE_READFILEOPTIONS_API_LATEST == 1, "EOS_TitleStorage_ReadFileOptions updated, check new fields");
+			Options.LocalUserId = GetProductUserIdChecked(Params.LocalAccountId);
+			Options.Filename = Utf8Filename.Get();
+			Options.ReadChunkLengthBytes = Config.ReadChunkLengthBytes;
+			Options.ReadFileDataCallback = &FTitleFileEOSGS::OnReadFileDataStatic;
+			Options.FileTransferProgressCallback = &FTitleFileEOSGS::OnFileTransferProgressStatic;
 
-		FTitleFileReadFileClientData* ClientData = new FTitleFileReadFileClientData(FileContentsRef, MoveTemp(Promise));
+			FTitleFileReadFileClientData* ClientData = new FTitleFileReadFileClientData(FileContentsRef, MoveTemp(Promise));
 
-		const EOS_HTitleStorageFileTransferRequest RequestHandle = EOS_TitleStorage_ReadFile(TitleStorageHandle, &Options, ClientData, &FTitleFileEOSGS::OnReadFileCompleteStatic);
-		InAsyncOp.Data.Set<EOS_HTitleStorageFileTransferRequest>(RequestHandleKey, RequestHandle);
-	})
-	.Then([this](TOnlineAsyncOp<FTitleFileReadFile>& InAsyncOp, const EOS_TitleStorage_ReadFileCallbackInfo* Data)
-	{
-		if (Data->ResultCode == EOS_EResult::EOS_Success)
+			const EOS_HTitleStorageFileTransferRequest RequestHandle = EOS_TitleStorage_ReadFile(TitleStorageHandle, &Options, ClientData, &FTitleFileEOSGS::OnReadFileCompleteStatic);
+			InAsyncOp.Data.Set<EOS_HTitleStorageFileTransferRequest>(RequestHandleKey, RequestHandle);
+		})
+		.Then([this](TOnlineAsyncOp<FTitleFileReadFile>& InAsyncOp, const EOS_TitleStorage_ReadFileCallbackInfo* Data)
 		{
-			const FTitleFileContentsRef* FileContents = InAsyncOp.Data.Get<FTitleFileContentsRef>(FileContentsKey);
-			if (ensure(FileContents))
+			if (Data->ResultCode == EOS_EResult::EOS_Success)
 			{
-				InAsyncOp.SetResult({ *FileContents });
+				const FTitleFileContentsRef* FileContents = InAsyncOp.Data.Get<FTitleFileContentsRef>(FileContentsKey);
+				if (ensure(FileContents))
+				{
+					InAsyncOp.SetResult({ *FileContents });
+				}
+				else
+				{
+					InAsyncOp.SetError(Errors::Unknown());
+				}
 			}
 			else
 			{
-				InAsyncOp.SetError(Errors::Unknown());
+				UE_LOG(LogTemp, Warning, TEXT("EOS_TitleStorage_ReadFile failed with result=[%s]"), *LexToString(Data->ResultCode));
+				InAsyncOp.SetError(Errors::FromEOSResult(Data->ResultCode));
 			}
-		}
-		else
-		{
-			UE_LOG(LogTemp, Warning, TEXT("EOS_TitleStorage_ReadFile failed with result=[%s]"), *LexToString(Data->ResultCode));
-			InAsyncOp.SetError(Errors::FromEOSResult(Data->ResultCode));
-		}
 
-		if (const EOS_HTitleStorageFileTransferRequest* RequestHandle = InAsyncOp.Data.Get<EOS_HTitleStorageFileTransferRequest>(RequestHandleKey))
-		{
-			EOS_TitleStorageFileTransferRequest_Release(*RequestHandle);
-		}
-	})
-	.Enqueue(GetSerialQueue());
+			if (const EOS_HTitleStorageFileTransferRequest* RequestHandle = InAsyncOp.Data.Get<EOS_HTitleStorageFileTransferRequest>(RequestHandleKey))
+			{
+				EOS_TitleStorageFileTransferRequest_Release(*RequestHandle);
+			}
+		})
+		.Enqueue(GetSerialQueue());
+	}
 	return Op->GetHandle();
 }
 
