- Use GetJoinableOp in EOS Leaderboards

diff --git a/Engine/Plugins/Online/Experimental/OnlineServicesEOSGS/Source/Private/Online/LeaderboardsEOSGS.cpp b/Engine/Plugins/Online/Experimental/OnlineServicesEOSGS/Source/Private/Online/LeaderboardsEOSGS.cpp
--- a/Engine/Plugins/Online/Experimental/OnlineServicesEOSGS/Source/Private/Online/LeaderboardsEOSGS.cpp
+++ b/Engine/Plugins/Online/Experimental/OnlineServicesEOSGS/Source/Private/Online/LeaderboardsEOSGS.cpp
@@ -30,90 +30,97 @@ void FLeaderboardsEOSGS::Initialize()
 
 TOnlineAsyncOpHandle<FReadEntriesForUsers> FLeaderboardsEOSGS::ReadEntriesForUsers(FReadEntriesForUsers::Params&& Params)
 {
-	TOnlineAsyncOpRef<FReadEntriesForUsers> Op = GetOp<FReadEntriesForUsers>(MoveTemp(Params));
-
-	if (!Services.Get<FAuthEOSGS>()->IsLoggedIn(Op->GetParams().LocalAccountId))
+	TOnlineAsyncOpRef<FReadEntriesForUsers> Op = GetJoinableOp<FReadEntriesForUsers>(MoveTemp(Params));
+	if (!Op->IsReady())
 	{
-		Op->SetError(Errors::InvalidUser());
-		return Op->GetHandle();
-	}
-
-	Op->Then([this](TOnlineAsyncOp<FReadEntriesForUsers>& InAsyncOp, TPromise<const EOS_Leaderboards_OnQueryLeaderboardUserScoresCompleteCallbackInfo*>&& Promise)
-	{
-		TArray<EOS_ProductUserId> ProductUserIds;
-		ProductUserIds.AddZeroed(InAsyncOp.GetParams().AccountIds.Num());
-		for (const FAccountId& AccountId : InAsyncOp.GetParams().AccountIds)
+		Op->Then([this](TOnlineAsyncOp<FReadEntriesForUsers>& InAsyncOp)
 		{
-			ProductUserIds.Emplace(GetProductUserIdChecked(AccountId));
-		}
-
-		EOS_Leaderboards_UserScoresQueryStatInfo StatInfo;
-		StatInfo.ApiVersion = EOS_LEADERBOARDS_USERSCORESQUERYSTATINFO_API_LATEST;
-		static_assert(EOS_LEADERBOARDS_USERSCORESQUERYSTATINFO_API_LATEST == 1, "EOS_Leaderboards_UserScoresQueryStatInfo updated, check new fields");
-
-		char StatNameANSI[EOS_LEADERBOARD_STAT_NAME_MAX_LENGTH];
-		FCStringAnsi::Strncpy(StatNameANSI, TCHAR_TO_UTF8(*InAsyncOp.GetParams().BoardName), EOS_LEADERBOARD_STAT_NAME_MAX_LENGTH);
-
-		StatInfo.StatName = StatNameANSI;
-		StatInfo.Aggregation = EOS_ELeaderboardAggregation::EOS_LA_Latest; // TODO: Use Stats definitions
-
-		EOS_Leaderboards_QueryLeaderboardUserScoresOptions Options = { };
-		Options.ApiVersion = EOS_LEADERBOARDS_QUERYLEADERBOARDUSERSCORES_API_LATEST;
-		static_assert(EOS_LEADERBOARDS_QUERYLEADERBOARDUSERSCORES_API_LATEST == 2, "EOS_Leaderboards_QueryLeaderboardUserScores updated, check new fields");
-		Options.UserIds = ProductUserIds.GetData();
-		Options.UserIdsCount = ProductUserIds.Num();
-		Options.StatInfo = &StatInfo;
-		Options.StatInfoCount = 1;
-		Options.StartTime = EOS_LEADERBOARDS_TIME_UNDEFINED;
-		Options.EndTime = EOS_LEADERBOARDS_TIME_UNDEFINED;
-		Options.LocalUserId = GetProductUserIdChecked(InAsyncOp.GetParams().LocalAccountId);
-
-		EOS_Async(EOS_Leaderboards_QueryLeaderboardUserScores, LeaderboardsHandle, Options, MoveTemp(Promise));
-	})
-	.Then([this](TOnlineAsyncOp<FReadEntriesForUsers>& InAsyncOp, const EOS_Leaderboards_OnQueryLeaderboardUserScoresCompleteCallbackInfo* Data)
-	{
-		if (Data->ResultCode != EOS_EResult::EOS_Success)
+			const FReadEntriesForUsers::Params& Params = InAsyncOp.GetParams();
+			if (!Services.Get<FAuthEOSGS>()->IsLoggedIn(Params.LocalAccountId))
+			{
+				InAsyncOp.SetError(Errors::InvalidUser());
+			}
+		})
+		.Then([this](TOnlineAsyncOp<FReadEntriesForUsers>& InAsyncOp, TPromise<const EOS_Leaderboards_OnQueryLeaderboardUserScoresCompleteCallbackInfo*>&& Promise)
 		{
-			UE_LOG(LogTemp, Warning, TEXT("EOS_Leaderboards_IngestStat failed with result=[%s]"), *LexToString(Data->ResultCode));
-			InAsyncOp.SetError(Errors::FromEOSResult(Data->ResultCode));
-			return;
-		}
-
-		char StatNameANSI[EOS_LEADERBOARD_STAT_NAME_MAX_LENGTH];
-		FCStringAnsi::Strncpy(StatNameANSI, TCHAR_TO_UTF8(*InAsyncOp.GetParams().BoardName), EOS_LEADERBOARD_STAT_NAME_MAX_LENGTH);
+			const FReadEntriesForUsers::Params& Params = InAsyncOp.GetParams();
 
-		FReadEntriesForUsers::Result Result;
+			TArray<EOS_ProductUserId> ProductUserIds;
+			ProductUserIds.AddZeroed(Params.AccountIds.Num());
+			for (const FAccountId& AccountId : Params.AccountIds)
+			{
+				ProductUserIds.Emplace(GetProductUserIdChecked(AccountId));
+			}
 
-		for (const FAccountId& AccountId : InAsyncOp.GetParams().AccountIds)
+			EOS_Leaderboards_UserScoresQueryStatInfo StatInfo;
+			StatInfo.ApiVersion = EOS_LEADERBOARDS_USERSCORESQUERYSTATINFO_API_LATEST;
+			static_assert(EOS_LEADERBOARDS_USERSCORESQUERYSTATINFO_API_LATEST == 1, "EOS_Leaderboards_UserScoresQueryStatInfo updated, check new fields");
+
+			char StatNameANSI[EOS_LEADERBOARD_STAT_NAME_MAX_LENGTH];
+			FCStringAnsi::Strncpy(StatNameANSI, TCHAR_TO_UTF8(*Params.BoardName), EOS_LEADERBOARD_STAT_NAME_MAX_LENGTH);
+
+			StatInfo.StatName = StatNameANSI;
+			StatInfo.Aggregation = EOS_ELeaderboardAggregation::EOS_LA_Latest; // TODO: Use Stats definitions
+
+			EOS_Leaderboards_QueryLeaderboardUserScoresOptions Options = { };
+			Options.ApiVersion = EOS_LEADERBOARDS_QUERYLEADERBOARDUSERSCORES_API_LATEST;
+			static_assert(EOS_LEADERBOARDS_QUERYLEADERBOARDUSERSCORES_API_LATEST == 2, "EOS_Leaderboards_QueryLeaderboardUserScores updated, check new fields");
+			Options.UserIds = ProductUserIds.GetData();
+			Options.UserIdsCount = ProductUserIds.Num();
+			Options.StatInfo = &StatInfo;
+			Options.StatInfoCount = 1;
+			Options.StartTime = EOS_LEADERBOARDS_TIME_UNDEFINED;
+			Options.EndTime = EOS_LEADERBOARDS_TIME_UNDEFINED;
+			Options.LocalUserId = GetProductUserIdChecked(Params.LocalAccountId);
+
+			EOS_Async(EOS_Leaderboards_QueryLeaderboardUserScores, LeaderboardsHandle, Options, MoveTemp(Promise));
+		})
+		.Then([this](TOnlineAsyncOp<FReadEntriesForUsers>& InAsyncOp, const EOS_Leaderboards_OnQueryLeaderboardUserScoresCompleteCallbackInfo* Data)
 		{
-			EOS_Leaderboards_CopyLeaderboardUserScoreByUserIdOptions UserCopyOptions = { };
-			UserCopyOptions.ApiVersion = EOS_LEADERBOARDS_COPYLEADERBOARDUSERSCOREBYUSERID_API_LATEST;
-			static_assert(EOS_LEADERBOARDS_COPYLEADERBOARDUSERSCOREBYUSERID_API_LATEST == 1, "EOS_Leaderboards_CopyLeaderboardUserScoreByUserIdOptions updated, check new fields");
+			if (Data->ResultCode != EOS_EResult::EOS_Success)
+			{
+				UE_LOG(LogTemp, Warning, TEXT("EOS_Leaderboards_IngestStat failed with result=[%s]"), *LexToString(Data->ResultCode));
+				InAsyncOp.SetError(Errors::FromEOSResult(Data->ResultCode));
+				return;
+			}
+
+			const FReadEntriesForUsers::Params& Params = InAsyncOp.GetParams();
 
-			UserCopyOptions.UserId = GetProductUserIdChecked(AccountId);
-			UserCopyOptions.StatName = StatNameANSI;
+			char StatNameANSI[EOS_LEADERBOARD_STAT_NAME_MAX_LENGTH];
+			FCStringAnsi::Strncpy(StatNameANSI, TCHAR_TO_UTF8(*Params.BoardName), EOS_LEADERBOARD_STAT_NAME_MAX_LENGTH);
 
-			EOS_Leaderboards_LeaderboardUserScore* LeaderboardUserScore = nullptr;
-			EOS_EResult CopyResult = EOS_Leaderboards_CopyLeaderboardUserScoreByUserId(LeaderboardsHandle, &UserCopyOptions, &LeaderboardUserScore);
+			FReadEntriesForUsers::Result Result;
 
-			if (CopyResult != EOS_EResult::EOS_Success)
+			for (const FAccountId& AccountId : Params.AccountIds)
 			{
-				UE_LOG(LogTemp, VeryVerbose, TEXT("Value not found for leaderboard %s."), *InAsyncOp.GetParams().BoardName);
-				continue;
-			}
+				EOS_Leaderboards_CopyLeaderboardUserScoreByUserIdOptions UserCopyOptions = { };
+				UserCopyOptions.ApiVersion = EOS_LEADERBOARDS_COPYLEADERBOARDUSERSCOREBYUSERID_API_LATEST;
+				static_assert(EOS_LEADERBOARDS_COPYLEADERBOARDUSERSCOREBYUSERID_API_LATEST == 1, "EOS_Leaderboards_CopyLeaderboardUserScoreByUserIdOptions updated, check new fields");
 
-			FLeaderboardEntry& Entry = Result.Entries.Emplace_GetRef();
-			Entry.Rank = UE_LEADERBOARD_RANK_UNKNOWN;
-			Entry.AccountId = AccountId;
-			Entry.Score = LeaderboardUserScore->Score;
+				UserCopyOptions.UserId = GetProductUserIdChecked(AccountId);
+				UserCopyOptions.StatName = StatNameANSI;
 
-			EOS_Leaderboards_LeaderboardUserScore_Release(LeaderboardUserScore);
-		}
+				EOS_Leaderboards_LeaderboardUserScore* LeaderboardUserScore = nullptr;
+				EOS_EResult CopyResult = EOS_Leaderboards_CopyLeaderboardUserScoreByUserId(LeaderboardsHandle, &UserCopyOptions, &LeaderboardUserScore);
 
-		InAsyncOp.SetResult(MoveTemp(Result));
-	}) 
-	.Enqueue(GetSerialQueue());
+				if (CopyResult != EOS_EResult::EOS_Success)
+				{
+					UE_LOG(LogTemp, VeryVerbose, TEXT("Value not found for leaderboard %s."), *Params.BoardName);
+					continue;
+				}
 
+				FLeaderboardEntry& Entry = Result.Entries.Emplace_GetRef();
+				Entry.Rank = UE_LEADERBOARD_RANK_UNKNOWN;
+				Entry.AccountId = AccountId;
+				Entry.Score = LeaderboardUserScore->Score;
+
+				EOS_Leaderboards_LeaderboardUserScore_Release(LeaderboardUserScore);
+			}
+
+			InAsyncOp.SetResult(MoveTemp(Result));
+		})
+		.Enqueue(GetSerialQueue());
+	}
 	return Op->GetHandle();
 }
 
@@ -161,97 +168,105 @@ void ReadEntriesInRange(EOS_HLeaderboards LeaderboardsHandle, uint32 StartIndex,
 
 TOnlineAsyncOpHandle<FReadEntriesAroundRank> FLeaderboardsEOSGS::ReadEntriesAroundRank(FReadEntriesAroundRank::Params&& Params)
 {
-	TOnlineAsyncOpRef<FReadEntriesAroundRank> Op = GetOp<FReadEntriesAroundRank>(MoveTemp(Params));
-
-	if (!Services.Get<FAuthEOSGS>()->IsLoggedIn(Op->GetParams().LocalAccountId))
+	TOnlineAsyncOpRef<FReadEntriesAroundRank> Op = GetJoinableOp<FReadEntriesAroundRank>(MoveTemp(Params));
+	if (!Op->IsReady())
 	{
-		Op->SetError(Errors::InvalidUser());
-		return Op->GetHandle();
-	}
-
-	if (Op->GetParams().Rank >= EOS_LEADERBOARD_MAX_ENTRIES || Op->GetParams().Limit == 0)
-	{
-		Op->SetError(Errors::InvalidParams());
-		return Op->GetHandle();
-	}
-
-	Op->Then([this](TOnlineAsyncOp<FReadEntriesAroundRank>& InAsyncOp, TPromise<const EOS_Leaderboards_OnQueryLeaderboardRanksCompleteCallbackInfo*>&& Promise)
-	{
-		Private::QueryLeaderboardsEOS(LeaderboardsHandle, InAsyncOp.GetParams().LocalAccountId, InAsyncOp.GetParams().BoardName, MoveTemp(Promise));
-	})
-	.Then([this](TOnlineAsyncOp<FReadEntriesAroundRank>& InAsyncOp, const EOS_Leaderboards_OnQueryLeaderboardRanksCompleteCallbackInfo* Data)
-	{
-		if (Data->ResultCode != EOS_EResult::EOS_Success)
+		Op->Then([this](TOnlineAsyncOp<FReadEntriesAroundRank>& InAsyncOp)
 		{
-			UE_LOG(LogTemp, Warning, TEXT("EOS_Leaderboards_QueryLeaderboardRanks failed with result=[%s]"), *LexToString(Data->ResultCode));
-			InAsyncOp.SetError(Errors::FromEOSResult(Data->ResultCode));
-			return;
-		}
+			const FReadEntriesAroundRank::Params& Params = InAsyncOp.GetParams();
+			if (!Services.Get<FAuthEOSGS>()->IsLoggedIn(Params.LocalAccountId))
+			{
+				InAsyncOp.SetError(Errors::InvalidUser());
+			}
+			else if (Params.Rank >= EOS_LEADERBOARD_MAX_ENTRIES || Params.Limit == 0)
+			{
+				InAsyncOp.SetError(Errors::InvalidParams());
+			}
+		})
+		.Then([this](TOnlineAsyncOp<FReadEntriesAroundRank>& InAsyncOp, TPromise<const EOS_Leaderboards_OnQueryLeaderboardRanksCompleteCallbackInfo*>&& Promise)
+		{
+			const FReadEntriesAroundRank::Params& Params = InAsyncOp.GetParams();
+			Private::QueryLeaderboardsEOS(LeaderboardsHandle, Params.LocalAccountId, Params.BoardName, MoveTemp(Promise));
+		})
+		.Then([this](TOnlineAsyncOp<FReadEntriesAroundRank>& InAsyncOp, const EOS_Leaderboards_OnQueryLeaderboardRanksCompleteCallbackInfo* Data)
+		{
+			if (Data->ResultCode != EOS_EResult::EOS_Success)
+			{
+				UE_LOG(LogTemp, Warning, TEXT("EOS_Leaderboards_QueryLeaderboardRanks failed with result=[%s]"), *LexToString(Data->ResultCode));
+				InAsyncOp.SetError(Errors::FromEOSResult(Data->ResultCode));
+				return;
+			}
 
-		FReadEntriesAroundRank::Result Result;
-		int32 StartIndex = InAsyncOp.GetParams().Rank;
-		int32 EndIndex = StartIndex + InAsyncOp.GetParams().Limit;
-		EndIndex = FMath::Clamp(EndIndex, 0, EOS_LEADERBOARD_MAX_ENTRIES-1);
-		Private::ReadEntriesInRange(LeaderboardsHandle, StartIndex, EndIndex, Result.Entries);
-		InAsyncOp.SetResult(MoveTemp(Result));
-	})
-	.Enqueue(GetSerialQueue());
+			const FReadEntriesAroundRank::Params& Params = InAsyncOp.GetParams();
 
+			FReadEntriesAroundRank::Result Result;
+			int32 StartIndex = Params.Rank;
+			int32 EndIndex = StartIndex + Params.Limit;
+			EndIndex = FMath::Clamp(EndIndex, 0, EOS_LEADERBOARD_MAX_ENTRIES-1);
+			Private::ReadEntriesInRange(LeaderboardsHandle, StartIndex, EndIndex, Result.Entries);
+			InAsyncOp.SetResult(MoveTemp(Result));
+		})
+		.Enqueue(GetSerialQueue());
+	}
 	return Op->GetHandle();
 }
 
 TOnlineAsyncOpHandle<FReadEntriesAroundUser> FLeaderboardsEOSGS::ReadEntriesAroundUser(FReadEntriesAroundUser::Params&& Params)
 {
-	TOnlineAsyncOpRef<FReadEntriesAroundUser> Op = GetOp<FReadEntriesAroundUser>(MoveTemp(Params));
-
-	if (!Services.Get<FAuthEOSGS>()->IsLoggedIn(Op->GetParams().LocalAccountId))
+	TOnlineAsyncOpRef<FReadEntriesAroundUser> Op = GetJoinableOp<FReadEntriesAroundUser>(MoveTemp(Params));
+	if (!Op->IsReady())
 	{
-		Op->SetError(Errors::InvalidUser());
-		return Op->GetHandle();
-	}
-
-	if (Op->GetParams().Limit > EOS_LEADERBOARD_MAX_ENTRIES || Op->GetParams().Limit == 0 || Op->GetParams().Offset >= EOS_LEADERBOARD_MAX_ENTRIES || Op->GetParams().Offset <= -EOS_LEADERBOARD_MAX_ENTRIES)
-	{
-		Op->SetError(Errors::InvalidParams());
-		return Op->GetHandle();
-	}
-
-	Op->Then([this](TOnlineAsyncOp<FReadEntriesAroundUser>& InAsyncOp, TPromise<const EOS_Leaderboards_OnQueryLeaderboardRanksCompleteCallbackInfo*>&& Promise)
-	{
-		Private::QueryLeaderboardsEOS(LeaderboardsHandle, InAsyncOp.GetParams().LocalAccountId, InAsyncOp.GetParams().BoardName, MoveTemp(Promise));
-	})
-	.Then([this](TOnlineAsyncOp<FReadEntriesAroundUser>& InAsyncOp, const EOS_Leaderboards_OnQueryLeaderboardRanksCompleteCallbackInfo* Data)
-	{
-		if (Data->ResultCode != EOS_EResult::EOS_Success)
+		Op->Then([this](TOnlineAsyncOp<FReadEntriesAroundUser>& InAsyncOp)
 		{
-			UE_LOG(LogTemp, Warning, TEXT("EOS_Leaderboards_QueryLeaderboardRanks failed with result=[%s]"), *LexToString(Data->ResultCode));
-			InAsyncOp.SetError(Errors::FromEOSResult(Data->ResultCode));
-			return;
-		}
+			const FReadEntriesAroundUser::Params& Params = InAsyncOp.GetParams();
+			if (!Services.Get<FAuthEOSGS>()->IsLoggedIn(Params.LocalAccountId))
+			{
+				InAsyncOp.SetError(Errors::InvalidUser());
+			}
+			else if (Params.Limit > EOS_LEADERBOARD_MAX_ENTRIES || Params.Limit == 0 || Params.Offset >= EOS_LEADERBOARD_MAX_ENTRIES || Params.Offset <= -EOS_LEADERBOARD_MAX_ENTRIES)
+			{
+				InAsyncOp.SetError(Errors::InvalidParams());
+			}
+		})
+		.Then([this](TOnlineAsyncOp<FReadEntriesAroundUser>& InAsyncOp, TPromise<const EOS_Leaderboards_OnQueryLeaderboardRanksCompleteCallbackInfo*>&& Promise)
+		{
+			const FReadEntriesAroundUser::Params& Params = InAsyncOp.GetParams();
+			Private::QueryLeaderboardsEOS(LeaderboardsHandle, Params.LocalAccountId, Params.BoardName, MoveTemp(Promise));
+		})
+		.Then([this](TOnlineAsyncOp<FReadEntriesAroundUser>& InAsyncOp, const EOS_Leaderboards_OnQueryLeaderboardRanksCompleteCallbackInfo* Data)
+		{
+			if (Data->ResultCode != EOS_EResult::EOS_Success)
+			{
+				UE_LOG(LogTemp, Warning, TEXT("EOS_Leaderboards_QueryLeaderboardRanks failed with result=[%s]"), *LexToString(Data->ResultCode));
+				InAsyncOp.SetError(Errors::FromEOSResult(Data->ResultCode));
+				return;
+			}
 
-		FReadEntriesAroundUser::Result Result;
+			const FReadEntriesAroundUser::Params& Params = InAsyncOp.GetParams();
 
-		EOS_Leaderboards_LeaderboardRecord* LeaderboardRecord = nullptr;
-		EOS_Leaderboards_CopyLeaderboardRecordByUserIdOptions CopyRecordByUserIdOptions;
-		CopyRecordByUserIdOptions.ApiVersion = EOS_LEADERBOARDS_COPYLEADERBOARDRECORDBYUSERID_API_LATEST;
-		static_assert(EOS_LEADERBOARDS_COPYLEADERBOARDRECORDBYUSERID_API_LATEST == 2, "EOS_Leaderboards_CopyLeaderboardRecordByUserIdOptions updated, check new fields");
-		CopyRecordByUserIdOptions.UserId = GetProductUserIdChecked(InAsyncOp.GetParams().AccountId);
-		EOS_EResult CopyResult = EOS_Leaderboards_CopyLeaderboardRecordByUserId(LeaderboardsHandle, &CopyRecordByUserIdOptions, &LeaderboardRecord);
-		if (CopyResult == EOS_EResult::EOS_Success)
-		{
-			int32 StartIndex = (int32)LeaderboardRecord->Rank + InAsyncOp.GetParams().Offset;
-			StartIndex = FMath::Clamp(StartIndex, 0, EOS_LEADERBOARD_MAX_ENTRIES-1);
-			int32 EndIndex = StartIndex + InAsyncOp.GetParams().Limit;
-			EndIndex = FMath::Clamp(EndIndex, 0, EOS_LEADERBOARD_MAX_ENTRIES-1);
-			Private::ReadEntriesInRange(LeaderboardsHandle, StartIndex, EndIndex, Result.Entries);
+			FReadEntriesAroundUser::Result Result;
 
-			EOS_Leaderboards_LeaderboardRecord_Release(LeaderboardRecord);
-		}
+			EOS_Leaderboards_LeaderboardRecord* LeaderboardRecord = nullptr;
+			EOS_Leaderboards_CopyLeaderboardRecordByUserIdOptions CopyRecordByUserIdOptions;
+			CopyRecordByUserIdOptions.ApiVersion = EOS_LEADERBOARDS_COPYLEADERBOARDRECORDBYUSERID_API_LATEST;
+			static_assert(EOS_LEADERBOARDS_COPYLEADERBOARDRECORDBYUSERID_API_LATEST == 2, "EOS_Leaderboards_CopyLeaderboardRecordByUserIdOptions updated, check new fields");
+			CopyRecordByUserIdOptions.UserId = GetProductUserIdChecked(Params.AccountId);
+			EOS_EResult CopyResult = EOS_Leaderboards_CopyLeaderboardRecordByUserId(LeaderboardsHandle, &CopyRecordByUserIdOptions, &LeaderboardRecord);
+			if (CopyResult == EOS_EResult::EOS_Success)
+			{
+				int32 StartIndex = (int32)LeaderboardRecord->Rank + Params.Offset;
+				StartIndex = FMath::Clamp(StartIndex, 0, EOS_LEADERBOARD_MAX_ENTRIES - 1);
+				int32 EndIndex = StartIndex + Params.Limit;
+				EndIndex = FMath::Clamp(EndIndex, 0, EOS_LEADERBOARD_MAX_ENTRIES - 1);
+				Private::ReadEntriesInRange(LeaderboardsHandle, StartIndex, EndIndex, Result.Entries);
 
-		InAsyncOp.SetResult(MoveTemp(Result));
-	})
-	.Enqueue(GetSerialQueue());
+				EOS_Leaderboards_LeaderboardRecord_Release(LeaderboardRecord);
+			}
 
+			InAsyncOp.SetResult(MoveTemp(Result));
+		})
+		.Enqueue(GetSerialQueue());
+	}
 	return Op->GetHandle();
 }
 
