- Use joinable ops in EOS stats

diff --git a/Engine/Plugins/Online/Experimental/OnlineServicesEOSGS/Source/Private/Online/StatsEOSGS.cpp b/Engine/Plugins/Online/Experimental/OnlineServicesEOSGS/Source/Private/Online/StatsEOSGS.cpp
--- a/Engine/Plugins/Online/Experimental/OnlineServicesEOSGS/Source/Private/Online/StatsEOSGS.cpp
+++ b/Engine/Plugins/Online/Experimental/OnlineServicesEOSGS/Source/Private/Online/StatsEOSGS.cpp
@@ -159,91 +159,97 @@ void ReadStatsFromEOSResult(EOS_HStats StatsHandle, const EOS_Stats_OnQueryStats
 
 TOnlineAsyncOpHandle<FQueryStats> FStatsEOSGS::QueryStats(FQueryStats::Params&& Params)
 {
-	TOnlineAsyncOpRef<FQueryStats> Op = GetOp<FQueryStats>(MoveTemp(Params));
-
-	if (!Services.Get<FAuthEOSGS>()->IsLoggedIn(Op->GetParams().LocalAccountId))
-	{
-		Op->SetError(Errors::InvalidUser());
-		return Op->GetHandle();
-	}
-
-	Op->Then([this](TOnlineAsyncOp<FQueryStats>& InAsyncOp, TPromise<const EOS_Stats_OnQueryStatsCompleteCallbackInfo*>&& Promise)
-	{
-		Private::QueryStatsEOS(StatsHandle, InAsyncOp.GetParams().LocalAccountId, InAsyncOp.GetParams().TargetAccountId, InAsyncOp.GetParams().StatNames, MoveTemp(Promise));
-	})
-	.Then([this](TOnlineAsyncOp<FQueryStats>& InAsyncOp, const EOS_Stats_OnQueryStatsCompleteCallbackInfo* Data)
+	TOnlineAsyncOpRef<FQueryStats> Op = GetJoinableOp<FQueryStats>(MoveTemp(Params));
+	if (!Op->IsReady())
 	{
-		if (Data->ResultCode != EOS_EResult::EOS_Success)
+		Op->Then([this](TOnlineAsyncOp<FQueryStats>& InAsyncOp)
 		{
-			UE_LOG(LogTemp, Warning, TEXT("EOS_Stats_QueryStats failed with result=[%s]"), *LexToString(Data->ResultCode));
-			InAsyncOp.SetError(Errors::FromEOSResult(Data->ResultCode));
-			return;
-		}
-
-		FQueryStats::Result Result;
-		Private::ReadStatsFromEOSResult(StatsHandle, Data, InAsyncOp.GetParams().StatNames, Result.Stats);
-
-		FUserStats UserStats;
-		UserStats.AccountId = InAsyncOp.GetParams().TargetAccountId;
-		UserStats.Stats = Result.Stats;
-		CacheUserStats(UserStats);
-
-		InAsyncOp.SetResult(MoveTemp(Result));
-	});
-
-	Op->Enqueue(GetSerialQueue());
-
-	return Op->GetHandle();
-}
-
-TOnlineAsyncOpHandle<FBatchQueryStats> FStatsEOSGS::BatchQueryStats(FBatchQueryStats::Params&& Params)
-{
-	TOnlineAsyncOpRef<FBatchQueryStats> Op = GetOp<FBatchQueryStats>(MoveTemp(Params));
-
-	if (!Services.Get<FAuthEOSGS>()->IsLoggedIn(Op->GetParams().LocalAccountId))
-	{
-		Op->SetError(Errors::InvalidUser());
-		return Op->GetHandle();
-	}
-
-	if (Op->GetParams().TargetAccountIds.IsEmpty() || Op->GetParams().StatNames.IsEmpty())
-	{
-		Op->SetError(Errors::InvalidParams());
-		return Op->GetHandle();
-	}
-
-	for (const FAccountId& TargetAccountId : Op->GetParams().TargetAccountIds)
-	{
-		Op->Then([this, &TargetAccountId](TOnlineAsyncOp<FBatchQueryStats>& InAsyncOp, TPromise<const EOS_Stats_OnQueryStatsCompleteCallbackInfo*>&& Promise)
+			const FQueryStats::Params& Params = InAsyncOp.GetParams();
+			if (!Services.Get<FAuthEOSGS>()->IsLoggedIn(Params.LocalAccountId))
+			{
+				InAsyncOp.SetError(Errors::InvalidUser());
+			}
+		})
+		.Then([this](TOnlineAsyncOp<FQueryStats>& InAsyncOp, TPromise<const EOS_Stats_OnQueryStatsCompleteCallbackInfo*>&& Promise)
 		{
-			Private::QueryStatsEOS(StatsHandle, InAsyncOp.GetParams().LocalAccountId, TargetAccountId, InAsyncOp.GetParams().StatNames, MoveTemp(Promise));
+			const FQueryStats::Params& Params = InAsyncOp.GetParams();
+			Private::QueryStatsEOS(StatsHandle, Params.LocalAccountId, Params.TargetAccountId, Params.StatNames, MoveTemp(Promise));
 		})
-		.Then([this](TOnlineAsyncOp<FBatchQueryStats>& InAsyncOp, const EOS_Stats_OnQueryStatsCompleteCallbackInfo* Data)
+		.Then([this](TOnlineAsyncOp<FQueryStats>& InAsyncOp, const EOS_Stats_OnQueryStatsCompleteCallbackInfo* Data)
 		{
 			if (Data->ResultCode != EOS_EResult::EOS_Success)
 			{
 				UE_LOG(LogTemp, Warning, TEXT("EOS_Stats_QueryStats failed with result=[%s]"), *LexToString(Data->ResultCode));
-				BatchQueriedUsersStats.Reset();
 				InAsyncOp.SetError(Errors::FromEOSResult(Data->ResultCode));
 				return;
 			}
 
-			FUserStats& UserStats = BatchQueriedUsersStats.Emplace_GetRef();
+			const FQueryStats::Params& Params = InAsyncOp.GetParams();
 
-			Private::ReadStatsFromEOSResult(StatsHandle, Data, InAsyncOp.GetParams().StatNames, UserStats.Stats);
+			FQueryStats::Result Result;
+			Private::ReadStatsFromEOSResult(StatsHandle, Data, Params.StatNames, Result.Stats);
+
+			FUserStats UserStats;
+			UserStats.AccountId = Params.TargetAccountId;
+			UserStats.Stats = Result.Stats;
 			CacheUserStats(UserStats);
-		});
+
+			InAsyncOp.SetResult(MoveTemp(Result));
+		})
+		.Enqueue(GetSerialQueue());
 	}
+	return Op->GetHandle();
+}
 
-	Op->Then([this](TOnlineAsyncOp<FBatchQueryStats>& InAsyncOp)
+TOnlineAsyncOpHandle<FBatchQueryStats> FStatsEOSGS::BatchQueryStats(FBatchQueryStats::Params&& Params)
+{
+	TOnlineAsyncOpRef<FBatchQueryStats> Op = GetJoinableOp<FBatchQueryStats>(MoveTemp(Params));
+	if (!Op->IsReady())
 	{
-		FBatchQueryStats::Result Result;
-		Result.UsersStats = MoveTemp(BatchQueriedUsersStats);
-		InAsyncOp.SetResult(MoveTemp(Result));
-	});
+		Op->Then([this](TOnlineAsyncOp<FBatchQueryStats>& InAsyncOp)
+		{
+			const FBatchQueryStats::Params& Params = InAsyncOp.GetParams();
+			if (!Services.Get<FAuthEOSGS>()->IsLoggedIn(Params.LocalAccountId))
+			{
+				InAsyncOp.SetError(Errors::InvalidUser());
+			}
+			else if (Params.TargetAccountIds.IsEmpty() || Params.StatNames.IsEmpty())
+			{
+				InAsyncOp.SetError(Errors::InvalidParams());
+			}
+		});
 
-	Op->Enqueue(GetSerialQueue());
+		for (const FAccountId& TargetAccountId : Op->GetParams().TargetAccountIds)
+		{
+			Op->Then([this, TargetAccountId](TOnlineAsyncOp<FBatchQueryStats>& InAsyncOp, TPromise<const EOS_Stats_OnQueryStatsCompleteCallbackInfo*>&& Promise)
+			{
+				Private::QueryStatsEOS(StatsHandle, InAsyncOp.GetParams().LocalAccountId, TargetAccountId, InAsyncOp.GetParams().StatNames, MoveTemp(Promise));
+			})
+			.Then([this](TOnlineAsyncOp<FBatchQueryStats>& InAsyncOp, const EOS_Stats_OnQueryStatsCompleteCallbackInfo* Data)
+			{
+				if (Data->ResultCode != EOS_EResult::EOS_Success)
+				{
+					UE_LOG(LogTemp, Warning, TEXT("EOS_Stats_QueryStats failed with result=[%s]"), *LexToString(Data->ResultCode));
+					BatchQueriedUsersStats.Reset();
+					InAsyncOp.SetError(Errors::FromEOSResult(Data->ResultCode));
+					return;
+				}
+
+				FUserStats& UserStats = BatchQueriedUsersStats.Emplace_GetRef();
+
+				Private::ReadStatsFromEOSResult(StatsHandle, Data, InAsyncOp.GetParams().StatNames, UserStats.Stats);
+				CacheUserStats(UserStats);
+			});
+		}
 
+		Op->Then([this](TOnlineAsyncOp<FBatchQueryStats>& InAsyncOp)
+		{
+			FBatchQueryStats::Result Result;
+			Result.UsersStats = MoveTemp(BatchQueriedUsersStats);
+			InAsyncOp.SetResult(MoveTemp(Result));
+		})
+		.Enqueue(GetSerialQueue());
+	}
 	return Op->GetHandle();
 }
 
