- Use joinable ops in EOS UserFile

diff --git a/Engine/Plugins/Online/Experimental/OnlineServicesEOSGS/Source/Private/Online/UserFileEOSGS.cpp b/Engine/Plugins/Online/Experimental/OnlineServicesEOSGS/Source/Private/Online/UserFileEOSGS.cpp
--- a/Engine/Plugins/Online/Experimental/OnlineServicesEOSGS/Source/Private/Online/UserFileEOSGS.cpp
+++ b/Engine/Plugins/Online/Experimental/OnlineServicesEOSGS/Source/Private/Online/UserFileEOSGS.cpp
@@ -38,85 +38,89 @@ void FUserFileEOSGS::UpdateConfig()
 
 TOnlineAsyncOpHandle<FUserFileEnumerateFiles> FUserFileEOSGS::EnumerateFiles(FUserFileEnumerateFiles::Params&& InParams)
 {
-	TOnlineAsyncOpRef<FUserFileEnumerateFiles> Op = GetOp<FUserFileEnumerateFiles>(MoveTemp(InParams));
-
-	Op->Then([this](TOnlineAsyncOp<FUserFileEnumerateFiles>& Op, TPromise<const EOS_PlayerDataStorage_QueryFileListCallbackInfo*>&& Promise)
+	TOnlineAsyncOpRef<FUserFileEnumerateFiles> Op = GetJoinableOp<FUserFileEnumerateFiles>(MoveTemp(InParams));
+	if (!Op->IsReady())
 	{
-		const FUserFileEnumerateFiles::Params& Params = Op.GetParams();
-
-		if (!Services.Get<FAuthEOSGS>()->IsLoggedIn(Params.LocalAccountId))
+		Op->Then([this](TOnlineAsyncOp<FUserFileEnumerateFiles>& Op)
 		{
-			Op.SetError(Errors::InvalidUser());
-			Promise.EmplaceValue();
-			return;
-		}
-
-		EOS_PlayerDataStorage_QueryFileListOptions Options = {};
-		Options.ApiVersion = EOS_PLAYERDATASTORAGE_QUERYFILELISTOPTIONS_API_LATEST;
-		static_assert(EOS_PLAYERDATASTORAGE_QUERYFILELISTOPTIONS_API_LATEST == 1, "EOS_PlayerDataStorage_QueryFileListOptions updated, check new fields");
-		Options.LocalUserId = GetProductUserIdChecked(Params.LocalAccountId);
+			const FUserFileEnumerateFiles::Params& Params = Op.GetParams();
 
-		EOS_Async(EOS_PlayerDataStorage_QueryFileList, PlayerDataStorageHandle, Options, MoveTemp(Promise));
-	})
-	.Then([this](TOnlineAsyncOp<FUserFileEnumerateFiles>& Op, const EOS_PlayerDataStorage_QueryFileListCallbackInfo* Data)
-	{
-		const FUserFileEnumerateFiles::Params& Params = Op.GetParams();
-
-		if (Data->ResultCode != EOS_EResult::EOS_Success
-			&& Data->ResultCode != EOS_EResult::EOS_NotFound)
+			if (!Services.Get<FAuthEOSGS>()->IsLoggedIn(Params.LocalAccountId))
+			{
+				Op.SetError(Errors::InvalidUser());
+			}
+		})
+		.Then([this](TOnlineAsyncOp<FUserFileEnumerateFiles>& Op, TPromise<const EOS_PlayerDataStorage_QueryFileListCallbackInfo*>&& Promise)
 		{
-			UE_LOG(LogTemp, Warning, TEXT("EOS_PlayerDataStorage_QueryFileList failed with result=[%s]"), *LexToString(Data->ResultCode));
-			Op.SetError(Errors::FromEOSResult(Data->ResultCode));
-			return;
-		}
-
-		const EOS_ProductUserId LocalUserPuid = GetProductUserIdChecked(Params.LocalAccountId);
+			const FUserFileEnumerateFiles::Params& Params = Op.GetParams();
 
-		EOS_PlayerDataStorage_GetFileMetadataCountOptions GetCountOptions;
-		GetCountOptions.ApiVersion = EOS_PLAYERDATASTORAGE_GETFILEMETADATACOUNTOPTIONS_API_LATEST;
-		static_assert(EOS_PLAYERDATASTORAGE_GETFILEMETADATACOUNTOPTIONS_API_LATEST == 1, "EOS_PlayerDataStorage_GetFileMetadataCountOptions updated, check new fields");
-		GetCountOptions.LocalUserId = LocalUserPuid;
+			EOS_PlayerDataStorage_QueryFileListOptions Options = {};
+			Options.ApiVersion = EOS_PLAYERDATASTORAGE_QUERYFILELISTOPTIONS_API_LATEST;
+			static_assert(EOS_PLAYERDATASTORAGE_QUERYFILELISTOPTIONS_API_LATEST == 1, "EOS_PlayerDataStorage_QueryFileListOptions updated, check new fields");
+			Options.LocalUserId = GetProductUserIdChecked(Params.LocalAccountId);
 
-		int32_t NumFiles;
-		EOS_EResult Result = EOS_PlayerDataStorage_GetFileMetadataCount(PlayerDataStorageHandle, &GetCountOptions, &NumFiles);
-		if (Result != EOS_EResult::EOS_Success)
+			EOS_Async(EOS_PlayerDataStorage_QueryFileList, PlayerDataStorageHandle, Options, MoveTemp(Promise));
+		})
+		.Then([this](TOnlineAsyncOp<FUserFileEnumerateFiles>& Op, const EOS_PlayerDataStorage_QueryFileListCallbackInfo* Data)
 		{
-			UE_LOG(LogTemp, Warning, TEXT("EOS_PlayerDataStorage_GetFileMetadataCount failed with result=[%s]"), *LexToString(Result));
-			Op.SetError(Errors::FromEOSResult(Result));
-			return;
-		}
+			const FUserFileEnumerateFiles::Params& Params = Op.GetParams();
 
-		check(NumFiles >= 0);
+			if (Data->ResultCode != EOS_EResult::EOS_Success
+				&& Data->ResultCode != EOS_EResult::EOS_NotFound)
+			{
+				UE_LOG(LogTemp, Warning, TEXT("EOS_PlayerDataStorage_QueryFileList failed with result=[%s]"), *LexToString(Data->ResultCode));
+				Op.SetError(Errors::FromEOSResult(Data->ResultCode));
+				return;
+			}
 
-		TArray<FString> NewEnumeratedFiles;
-		NewEnumeratedFiles.Reserve(NumFiles);
+			const EOS_ProductUserId LocalUserPuid = GetProductUserIdChecked(Params.LocalAccountId);
 
-		for (int32_t FileIdx = 0; FileIdx < NumFiles; FileIdx++)
-		{
-			EOS_PlayerDataStorage_CopyFileMetadataAtIndexOptions CopyOptions;
-			CopyOptions.ApiVersion = EOS_PLAYERDATASTORAGE_COPYFILEMETADATAATINDEXOPTIONS_API_LATEST;
-			static_assert(EOS_PLAYERDATASTORAGE_COPYFILEMETADATAATINDEXOPTIONS_API_LATEST == 1, "EOS_PlayerDataStorage_CopyFileMetadataAtIndexOptions updated, check new fields");
-			CopyOptions.LocalUserId = LocalUserPuid;
-			CopyOptions.Index = FileIdx;
-
-			EOS_PlayerDataStorage_FileMetadata* FileMetadata = nullptr;
-			Result = EOS_PlayerDataStorage_CopyFileMetadataAtIndex(PlayerDataStorageHandle, &CopyOptions, &FileMetadata);
+			EOS_PlayerDataStorage_GetFileMetadataCountOptions GetCountOptions;
+			GetCountOptions.ApiVersion = EOS_PLAYERDATASTORAGE_GETFILEMETADATACOUNTOPTIONS_API_LATEST;
+			static_assert(EOS_PLAYERDATASTORAGE_GETFILEMETADATACOUNTOPTIONS_API_LATEST == 1, "EOS_PlayerDataStorage_GetFileMetadataCountOptions updated, check new fields");
+			GetCountOptions.LocalUserId = LocalUserPuid;
+
+			int32_t NumFiles;
+			EOS_EResult Result = EOS_PlayerDataStorage_GetFileMetadataCount(PlayerDataStorageHandle, &GetCountOptions, &NumFiles);
 			if (Result != EOS_EResult::EOS_Success)
 			{
-				UE_LOG(LogTemp, Warning, TEXT("EOS_PlayerDataStorage_CopyFileMetadataAtIndex failed with result=[%s]"), *LexToString(Result));
+				UE_LOG(LogTemp, Warning, TEXT("EOS_PlayerDataStorage_GetFileMetadataCount failed with result=[%s]"), *LexToString(Result));
 				Op.SetError(Errors::FromEOSResult(Result));
 				return;
 			}
 
-			NewEnumeratedFiles.Emplace(UTF8_TO_TCHAR(FileMetadata->Filename));
+			check(NumFiles >= 0);
 
-			EOS_PlayerDataStorage_FileMetadata_Release(FileMetadata);
-		}
+			TArray<FString> NewEnumeratedFiles;
+			NewEnumeratedFiles.Reserve(NumFiles);
 
-		UserToFiles.Emplace(Params.LocalAccountId, MoveTemp(NewEnumeratedFiles));
-		Op.SetResult({});
-	})
-	.Enqueue(GetSerialQueue());
+			for (int32_t FileIdx = 0; FileIdx < NumFiles; FileIdx++)
+			{
+				EOS_PlayerDataStorage_CopyFileMetadataAtIndexOptions CopyOptions;
+				CopyOptions.ApiVersion = EOS_PLAYERDATASTORAGE_COPYFILEMETADATAATINDEXOPTIONS_API_LATEST;
+				static_assert(EOS_PLAYERDATASTORAGE_COPYFILEMETADATAATINDEXOPTIONS_API_LATEST == 1, "EOS_PlayerDataStorage_CopyFileMetadataAtIndexOptions updated, check new fields");
+				CopyOptions.LocalUserId = LocalUserPuid;
+				CopyOptions.Index = FileIdx;
+
+				EOS_PlayerDataStorage_FileMetadata* FileMetadata = nullptr;
+				Result = EOS_PlayerDataStorage_CopyFileMetadataAtIndex(PlayerDataStorageHandle, &CopyOptions, &FileMetadata);
+				if (Result != EOS_EResult::EOS_Success)
+				{
+					UE_LOG(LogTemp, Warning, TEXT("EOS_PlayerDataStorage_CopyFileMetadataAtIndex failed with result=[%s]"), *LexToString(Result));
+					Op.SetError(Errors::FromEOSResult(Result));
+					return;
+				}
+
+				NewEnumeratedFiles.Emplace(UTF8_TO_TCHAR(FileMetadata->Filename));
+
+				EOS_PlayerDataStorage_FileMetadata_Release(FileMetadata);
+			}
+
+			UserToFiles.Emplace(Params.LocalAccountId, MoveTemp(NewEnumeratedFiles));
+			Op.SetResult({});
+		})
+		.Enqueue(GetSerialQueue());
+	}
 	return Op->GetHandle();
 }
 
@@ -154,71 +158,74 @@ TOnlineAsyncOpHandle<FUserFileReadFile> FUserFileEOSGS::ReadFile(FUserFileReadFi
 {
 	static const TCHAR* FileContentsKey = TEXT("FileContents");
 
-	TOnlineAsyncOpRef<FUserFileReadFile> Op = GetOp<FUserFileReadFile>(MoveTemp(InParams));
-	const FUserFileReadFile::Params& Params = Op->GetParams();
-
-	if (Params.Filename.IsEmpty())
-	{
-		Op->SetError(Errors::InvalidParams());
-		return Op->GetHandle();
-	}
-
-	Op->Then([this](TOnlineAsyncOp<FUserFileReadFile>& Op, TPromise<const EOS_PlayerDataStorage_ReadFileCallbackInfo*>&& Promise)
+	TOnlineAsyncOpRef<FUserFileReadFile> Op = GetJoinableOp<FUserFileReadFile>(MoveTemp(InParams));
+	if (!Op->IsReady())
 	{
-		const FUserFileReadFile::Params& Params = Op.GetParams();
-
-		if (!Services.Get<FAuthEOSGS>()->IsLoggedIn(Params.LocalAccountId))
+		const FUserFileReadFile::Params& Params = Op->GetParams();
+		if (Params.Filename.IsEmpty())
 		{
-			Op.SetError(Errors::InvalidUser());
-			Promise.EmplaceValue();
-			return;
+			Op->SetError(Errors::InvalidParams());
+			return Op->GetHandle();
 		}
 
-		FUserFileContentsRef FileContentsRef = MakeShared<FUserFileContents>();
-		Op.Data.Set<FUserFileContentsRef>(FileContentsKey, FileContentsRef);
+		Op->Then([this](TOnlineAsyncOp<FUserFileReadFile>& Op)
+		{
+			const FUserFileReadFile::Params& Params = Op.GetParams();
+			if (!Services.Get<FAuthEOSGS>()->IsLoggedIn(Params.LocalAccountId))
+			{
+				Op.SetError(Errors::InvalidUser());
+			}
+		})
+		.Then([this](TOnlineAsyncOp<FUserFileReadFile>& Op, TPromise<const EOS_PlayerDataStorage_ReadFileCallbackInfo*>&& Promise)
+		{
+			const FUserFileReadFile::Params& Params = Op.GetParams();
 
-		const auto Utf8Filename = StringCast<UTF8CHAR>(*Params.Filename);
+			FUserFileContentsRef FileContentsRef = MakeShared<FUserFileContents>();
+			Op.Data.Set<FUserFileContentsRef>(FileContentsKey, FileContentsRef);
 
-		EOS_PlayerDataStorage_ReadFileOptions Options = {};
-		Options.ApiVersion = EOS_PLAYERDATASTORAGE_READFILEOPTIONS_API_LATEST;
-		static_assert(EOS_PLAYERDATASTORAGE_READFILEOPTIONS_API_LATEST == 1, "EOS_PlayerDataStorage_ReadFileOptions updated, check new fields");
-		Options.LocalUserId = GetProductUserIdChecked(Params.LocalAccountId);
-		Options.Filename = (const char*)Utf8Filename.Get();
-		Options.ReadChunkLengthBytes = Config.ChunkLengthBytes;
-		Options.ReadFileDataCallback = &FUserFileEOSGS::OnReadFileDataStatic;
-		Options.FileTransferProgressCallback = &FUserFileEOSGS::OnFileTransferProgressStatic;
+			const auto Utf8Filename = StringCast<UTF8CHAR>(*Params.Filename);
 
-		FUserFileReadFileClientData* ClientData = new FUserFileReadFileClientData(FileContentsRef, MoveTemp(Promise));
+			EOS_PlayerDataStorage_ReadFileOptions Options = {};
+			Options.ApiVersion = EOS_PLAYERDATASTORAGE_READFILEOPTIONS_API_LATEST;
+			static_assert(EOS_PLAYERDATASTORAGE_READFILEOPTIONS_API_LATEST == 1, "EOS_PlayerDataStorage_ReadFileOptions updated, check new fields");
+			Options.LocalUserId = GetProductUserIdChecked(Params.LocalAccountId);
+			Options.Filename = (const char*)Utf8Filename.Get();
+			Options.ReadChunkLengthBytes = Config.ChunkLengthBytes;
+			Options.ReadFileDataCallback = &FUserFileEOSGS::OnReadFileDataStatic;
+			Options.FileTransferProgressCallback = &FUserFileEOSGS::OnFileTransferProgressStatic;
 
-		const EOS_HPlayerDataStorageFileTransferRequest RequestHandle = EOS_PlayerDataStorage_ReadFile(PlayerDataStorageHandle, &Options, ClientData, &FUserFileEOSGS::OnReadFileCompleteStatic);
-		Op.Data.Set<EOS_HPlayerDataStorageFileTransferRequest>(RequestHandleKey, RequestHandle);
-	})
-	.Then([this](TOnlineAsyncOp<FUserFileReadFile>& Op, const EOS_PlayerDataStorage_ReadFileCallbackInfo* Data)
-	{
-		if (Data->ResultCode == EOS_EResult::EOS_Success)
+			FUserFileReadFileClientData* ClientData = new FUserFileReadFileClientData(FileContentsRef, MoveTemp(Promise));
+
+			const EOS_HPlayerDataStorageFileTransferRequest RequestHandle = EOS_PlayerDataStorage_ReadFile(PlayerDataStorageHandle, &Options, ClientData, &FUserFileEOSGS::OnReadFileCompleteStatic);
+			Op.Data.Set<EOS_HPlayerDataStorageFileTransferRequest>(RequestHandleKey, RequestHandle);
+		})
+		.Then([this](TOnlineAsyncOp<FUserFileReadFile>& Op, const EOS_PlayerDataStorage_ReadFileCallbackInfo* Data)
 		{
-			const FUserFileContentsRef* FileContents = Op.Data.Get<FUserFileContentsRef>(FileContentsKey);
-			if (ensure(FileContents))
+			if (Data->ResultCode == EOS_EResult::EOS_Success)
 			{
-				Op.SetResult({ *FileContents });
+				const FUserFileContentsRef* FileContents = Op.Data.Get<FUserFileContentsRef>(FileContentsKey);
+				if (ensure(FileContents))
+				{
+					Op.SetResult({ *FileContents });
+				}
+				else
+				{
+					Op.SetError(Errors::Unknown());
+				}
 			}
 			else
 			{
-				Op.SetError(Errors::Unknown());
+				UE_LOG(LogTemp, Warning, TEXT("EOS_PlayerDataStorage_ReadFile failed with result=[%s]"), *LexToString(Data->ResultCode));
+				Op.SetError(Errors::FromEOSResult(Data->ResultCode));
 			}
-		}
-		else
-		{
-			UE_LOG(LogTemp, Warning, TEXT("EOS_PlayerDataStorage_ReadFile failed with result=[%s]"), *LexToString(Data->ResultCode));
-			Op.SetError(Errors::FromEOSResult(Data->ResultCode));
-		}
 
-		if (const EOS_HPlayerDataStorageFileTransferRequest* RequestHandle = Op.Data.Get<EOS_HPlayerDataStorageFileTransferRequest>(RequestHandleKey))
-		{
-			EOS_PlayerDataStorageFileTransferRequest_Release(*RequestHandle);
-		}
-	})
-	.Enqueue(GetSerialQueue());
+			if (const EOS_HPlayerDataStorageFileTransferRequest* RequestHandle = Op.Data.Get<EOS_HPlayerDataStorageFileTransferRequest>(RequestHandleKey))
+			{
+				EOS_PlayerDataStorageFileTransferRequest_Release(*RequestHandle);
+			}
+		})
+		.Enqueue(GetSerialQueue());
+	}
 	return Op->GetHandle();
 }
 
