Implementation or removal of pending TODOs in session interface after evaluation (Part 1)
Added SetHostAddress to EOSGS implementation to expose EOSSDK method
Added Id resolution logic to Add/RemoveSessionMember in EOSGS implementation

diff --git a/Engine/Plugins/Online/Experimental/OnlineServices/Source/OnlineServicesCommonEngineUtils/Private/Online/SessionsLAN.cpp b/Engine/Plugins/Online/Experimental/OnlineServices/Source/OnlineServicesCommonEngineUtils/Private/Online/SessionsLAN.cpp
--- a/Engine/Plugins/Online/Experimental/OnlineServices/Source/OnlineServicesCommonEngineUtils/Private/Online/SessionsLAN.cpp
+++ b/Engine/Plugins/Online/Experimental/OnlineServices/Source/OnlineServicesCommonEngineUtils/Private/Online/SessionsLAN.cpp
@@ -305,14 +305,14 @@ TOptional<FOnlineError> FSessionsLAN::TryHostLANSession()
 
 			LANSessionManager->StopLANSession();
 
-			Result.Emplace(Errors::RequestFailure()); // TODO: May need a new error type
+			Result.Emplace(Errors::RequestFailure());
 		}
 	}
 	else
 	{
 		UE_LOG(LogTemp, VeryVerbose, TEXT("[FSessionsLAN::TryHostLANSession] LAN Beacon already in use!"));
 
-		Result.Emplace(Errors::RequestFailure()); // TODO: May need a new error type
+		Result.Emplace(Errors::AlreadyPending());
 	}
 
 	return Result;
@@ -339,7 +339,7 @@ void FSessionsLAN::FindLANSessions(const FAccountId& LocalAccountId)
 		}
 
 		// Trigger the delegate as having failed
-		CurrentSessionSearchPromisesUserMap.FindChecked(LocalAccountId).EmplaceValue(Errors::RequestFailure()); // TODO: May need a new error type
+		CurrentSessionSearchPromisesUserMap.FindChecked(LocalAccountId).EmplaceValue(Errors::RequestFailure());
 		CurrentSessionSearchPromisesUserMap.Remove(LocalAccountId);
 
 		// If we were hosting public sessions before the search, we'll return the beacon to that state
diff --git a/Engine/Plugins/Online/Experimental/OnlineServicesEOSGS/Source/Private/Online/SessionsEOSGS.cpp b/Engine/Plugins/Online/Experimental/OnlineServicesEOSGS/Source/Private/Online/SessionsEOSGS.cpp
--- a/Engine/Plugins/Online/Experimental/OnlineServicesEOSGS/Source/Private/Online/SessionsEOSGS.cpp
+++ b/Engine/Plugins/Online/Experimental/OnlineServicesEOSGS/Source/Private/Online/SessionsEOSGS.cpp
@@ -54,11 +54,10 @@ FSessionEOSGS::FSessionEOSGS(const TSharedPtr<FSessionDetailsHandleEOSGS>& InSes
 	if (CopyInfoResult == EOS_EResult::EOS_Success)
 	{
 		SessionInfo.SessionId = FSessionsEOSGS::CreateSessionId(FString(SessionDetailsInfo->SessionId));
-		SessionInfo.bAllowSanctionedPlayers = !SessionDetailsInfo->Settings->bSanctionsEnabled;
-
-		// We could retrieve the Host Address here if we set it during creation or update
+		SessionSettings.CustomSettings.Emplace(EOSGS_HOST_ADDRESS_ATTRIBUTE_KEY, FCustomSessionSetting{ FSchemaVariant(FString(UTF8_TO_TCHAR(SessionDetailsInfo->HostAddress))), ESchemaAttributeVisibility::Public });
 
 		// bInvitesAllowed and bAllowJoinInProgress should both have the same value, and that value corresponds to bAllowNewMembers
+		SessionInfo.bAllowSanctionedPlayers = !SessionDetailsInfo->Settings->bSanctionsEnabled;
 		SessionSettings.bAllowNewMembers = SessionDetailsInfo->Settings->bInvitesAllowed && SessionDetailsInfo->Settings->bAllowJoinInProgress;
 		SessionSettings.NumMaxConnections = SessionDetailsInfo->Settings->NumPublicConnections;
 		SessionSettings.JoinPolicy = FromServiceType(SessionDetailsInfo->Settings->PermissionLevel);
@@ -341,16 +340,6 @@ TFuture<TOnlineResult<FCreateSession>> FSessionsEOSGS::CreateSessionImpl(const F
 		}
 	}
 
-	// Check if the Bucket Id custom setting is set. EOS Sessions can not be created without it
-	if (!Params.SessionSettings.CustomSettings.Contains(EOSGS_BUCKET_ID_ATTRIBUTE_KEY))
-	{
-		UE_LOG(LogTemp, Warning, TEXT("[FSessionsEOSGS::CreateSession] Could not create session without Custom Setting 'EOSGS_BUCKET_ID_ATTRIBUTE_KEY' (FString) set."));
-
-		Promise.EmplaceValue(Errors::InvalidParams());
-
-		return Future;
-	}
-
 	// After all initial checks, we start the session creation operations
 
 	EOS_Sessions_CreateSessionModificationOptions CreateSessionModificationOptions = {};
@@ -360,13 +349,23 @@ TFuture<TOnlineResult<FCreateSession>> FSessionsEOSGS::CreateSessionImpl(const F
 	CreateSessionModificationOptions.bPresenceEnabled = Params.bPresenceEnabled;
 	CreateSessionModificationOptions.bSanctionsEnabled = !Params.bAllowSanctionedPlayers;
 
-	const FCustomSessionSetting* BucketIdSetting = Params.SessionSettings.CustomSettings.Find(EOSGS_BUCKET_ID_ATTRIBUTE_KEY);
-	const FTCHARToUTF8 BucketIdUtf8(BucketIdSetting ? *BucketIdSetting->Data.GetString() : nullptr);
-
-	if (BucketIdUtf8.Length())
+	// Check if the Bucket Id custom setting is set, set default otherwise. EOS Sessions can not be created without it
+	FString BucketIdStr;
+	if (const FCustomSessionSetting* BucketIdSetting = Params.SessionSettings.CustomSettings.Find(EOSGS_BUCKET_ID_ATTRIBUTE_KEY))
 	{
-		CreateSessionModificationOptions.BucketId = BucketIdUtf8.Get();
+		BucketIdStr = BucketIdSetting->Data.GetString();
 	}
+	else
+	{
+		int32 BuildUniqueId = GetBuildUniqueId();
+
+		UE_LOG(LogTemp, Verbose, TEXT("[FSessionsEOSGS::CreateSession] 'EOSGS_BUCKET_ID_ATTRIBUTE_KEY' (FString) Custom Setting needed to create EOS sessions not found. Setting \"%d\" as default."), BuildUniqueId);
+
+		BucketIdStr = FString::FromInt(BuildUniqueId);
+	}
+
+	const FTCHARToUTF8 BucketIdUtf8(*BucketIdStr);
+	CreateSessionModificationOptions.BucketId = BucketIdUtf8.Get();
 	
 	CreateSessionModificationOptions.LocalUserId = GetProductUserIdChecked(Params.LocalAccountId);
 	CreateSessionModificationOptions.MaxPlayers = Params.SessionSettings.NumMaxConnections;
@@ -391,8 +390,6 @@ TFuture<TOnlineResult<FCreateSession>> FSessionsEOSGS::CreateSessionImpl(const F
 		return Future;
 	}
 
-	// TODO: We could call EOS_SessionModification_SetHostAddress at this point, although it's not necessary
-
 	// We write all SessionSettings values into the SessionModificationHandle
 	WriteCreateSessionModificationHandle(SessionModificationHandle, Params);
 
@@ -428,6 +425,23 @@ TFuture<TOnlineResult<FCreateSession>> FSessionsEOSGS::CreateSessionImpl(const F
 	return Future;
 }
 
+void FSessionsEOSGS::SetHostAddress(EOS_HSessionModification& SessionModHandle, FString HostAddress)
+{
+	// TODO: We could call EOS_SessionModification_SetHostAddress at this point, although it's not necessary
+	EOS_SessionModification_SetHostAddressOptions Options = {};
+	Options.ApiVersion = EOS_SESSIONMODIFICATION_SETHOSTADDRESS_API_LATEST;
+	static_assert(EOS_SESSIONMODIFICATION_SETHOSTADDRESS_API_LATEST == 1, "EOS_SessionModification_SetHostAddressOptions updated, check new fields");
+
+	const FTCHARToUTF8 HostAddressUtf8(*HostAddress);
+	Options.HostAddress = HostAddressUtf8.Get();
+
+	EOS_EResult ResultCode = EOS_SessionModification_SetHostAddress(SessionModHandle, &Options);
+	if (ResultCode != EOS_EResult::EOS_Success)
+	{
+		UE_LOG(LogTemp, Warning, TEXT("EOS_SessionModification_SetHostAddress failed with result [%s]"), *LexToString(ResultCode));
+	}
+}
+
 void FSessionsEOSGS::SetJoinInProgressAllowed(EOS_HSessionModification& SessionModHandle, bool bIsJoinInProgressAllowed)
 {
 	EOS_SessionModification_SetJoinInProgressAllowedOptions Options = {};
@@ -540,7 +554,15 @@ void FSessionsEOSGS::RemoveAttribute(EOS_HSessionModification& SessionModificati
 
 void FSessionsEOSGS::WriteCreateSessionModificationHandle(EOS_HSessionModification& SessionModificationHandle, const FCreateSession::Params& Params)
 {
-	// TODO: We have the option to call EOS_SessionModification_SetHostAddress in EOS, useful if the session owner changes
+	// We'll enable setting a custom host address if a value for it was set as a custom parameter
+	if (const FCustomSessionSetting* HostAddressSetting = Params.SessionSettings.CustomSettings.Find(EOSGS_HOST_ADDRESS_ATTRIBUTE_KEY))
+	{
+		FString HostAddressStr = HostAddressSetting->Data.GetString();
+		if (!HostAddressStr.IsEmpty())
+		{
+			SetHostAddress(SessionModificationHandle, HostAddressStr);
+		}
+	}
 
 	// We'll update this setting in the session modification step, and start or end the session accordingly to bock or allow join processes
 	SetJoinInProgressAllowed(SessionModificationHandle, Params.SessionSettings.bAllowNewMembers);
@@ -580,7 +602,15 @@ void FSessionsEOSGS::WriteCreateSessionModificationHandle(EOS_HSessionModificati
 
 void FSessionsEOSGS::WriteUpdateSessionModificationHandle(EOS_HSessionModification& SessionModificationHandle, const FSessionSettingsUpdate& NewSettings)
 {
-	// TODO: We have the option to call EOS_SessionModification_SetHostAddress in EOS, useful if the session owner changes
+	// We'll enable setting a custom host address if a value for it was set as a custom parameter
+	if (const FCustomSessionSetting* HostAddressSetting = NewSettings.UpdatedCustomSettings.Find(EOSGS_HOST_ADDRESS_ATTRIBUTE_KEY))
+	{
+		FString HostAddressStr = HostAddressSetting->Data.GetString();
+		if (!HostAddressStr.IsEmpty())
+		{
+			SetHostAddress(SessionModificationHandle, HostAddressStr);
+		}
+	}
 
 	if (NewSettings.bAllowNewMembers.IsSet())
 	{
@@ -1489,30 +1519,40 @@ TFuture<TOnlineResult<FAddSessionMember>> FSessionsEOSGS::AddSessionMemberImpl(c
 	TPromise<TOnlineResult<FAddSessionMember>> Promise;
 	TFuture<TOnlineResult<FAddSessionMember>> Future = Promise.GetFuture();
 
-	EOS_Sessions_RegisterPlayersOptions Options = { };
-	Options.ApiVersion = EOS_SESSIONS_REGISTERPLAYERS_API_LATEST;
-	static_assert(EOS_SESSIONS_REGISTERPLAYERS_API_LATEST == 2, "EOS_Sessions_RegisterPlayersOptions updated, check new fields");
+	Services.Get<FAuthEOSGS>()->ResolveAccountIds(Params.LocalAccountId, { GetProductUserIdChecked(Params.LocalAccountId) })
+		.Next([this, Promise = MoveTemp(Promise), Params](TArray<FAccountId> ResolvedAccountIds) mutable
+			{
+				if (!ResolvedAccountIds.IsEmpty())
+				{
+					EOS_Sessions_RegisterPlayersOptions Options = { };
+					Options.ApiVersion = EOS_SESSIONS_REGISTERPLAYERS_API_LATEST;
+					static_assert(EOS_SESSIONS_REGISTERPLAYERS_API_LATEST == 2, "EOS_Sessions_RegisterPlayersOptions updated, check new fields");
 
-	Options.PlayersToRegisterCount = 1;
+					Options.PlayersToRegisterCount = 1;
 
-	// TODO: Do this with ResolveAccountIds instead
-	TArray<EOS_ProductUserId> PlayersToRegister = { GetProductUserIdChecked(Params.LocalAccountId) };
-	Options.PlayersToRegister = PlayersToRegister.GetData();
+					TArray<EOS_ProductUserId> PlayersToRegister = { GetProductUserIdChecked(Params.LocalAccountId) };
+					Options.PlayersToRegister = PlayersToRegister.GetData();
 
-	const FTCHARToUTF8 SessionNameUtf8(*Params.SessionName.ToString());
-	Options.SessionName = SessionNameUtf8.Get();
+					const FTCHARToUTF8 SessionNameUtf8(*Params.SessionName.ToString());
+					Options.SessionName = SessionNameUtf8.Get();
 
-	EOS_Async(EOS_Sessions_RegisterPlayers, SessionsHandle, Options,
-		[this, Promise = MoveTemp(Promise), Params](const EOS_Sessions_RegisterPlayersCallbackInfo* Result) mutable
-		{
-			if (Result->ResultCode != EOS_EResult::EOS_Success)
-			{
-				Promise.EmplaceValue(Errors::FromEOSResult(Result->ResultCode));
-				return;
-			}
-
-			Promise.EmplaceValue(AddSessionMemberInternal(Params));
-		});
+					EOS_Async(EOS_Sessions_RegisterPlayers, SessionsHandle, Options,
+						[this, Promise = MoveTemp(Promise), Params](const EOS_Sessions_RegisterPlayersCallbackInfo* Result) mutable
+						{
+							if (Result->ResultCode != EOS_EResult::EOS_Success)
+							{
+								Promise.EmplaceValue(Errors::FromEOSResult(Result->ResultCode));
+								return;
+							}
+
+							Promise.EmplaceValue(AddSessionMemberInternal(Params));
+						});
+				}
+				else
+				{
+					Promise.EmplaceValue(Errors::InvalidUser());
+				}
+			});
 
 	return Future;
 }
@@ -1531,30 +1571,40 @@ TFuture<TOnlineResult<FRemoveSessionMember>> FSessionsEOSGS::RemoveSessionMember
 	TPromise<TOnlineResult<FRemoveSessionMember>> Promise;
 	TFuture<TOnlineResult<FRemoveSessionMember>> Future = Promise.GetFuture();
 
-	EOS_Sessions_UnregisterPlayersOptions Options = { };
-	Options.ApiVersion = EOS_SESSIONS_UNREGISTERPLAYERS_API_LATEST;
-	static_assert(EOS_SESSIONS_UNREGISTERPLAYERS_API_LATEST == 2, "EOS_Sessions_UnregisterPlayersOptions updated, check new fields");
+	Services.Get<FAuthEOSGS>()->ResolveAccountIds(Params.LocalAccountId, { GetProductUserIdChecked(Params.LocalAccountId) })
+		.Next([this, Promise = MoveTemp(Promise), Params](TArray<FAccountId> ResolvedAccountIds) mutable
+			{
+				if (!ResolvedAccountIds.IsEmpty())
+				{
+					EOS_Sessions_UnregisterPlayersOptions Options = { };
+					Options.ApiVersion = EOS_SESSIONS_UNREGISTERPLAYERS_API_LATEST;
+					static_assert(EOS_SESSIONS_UNREGISTERPLAYERS_API_LATEST == 2, "EOS_Sessions_UnregisterPlayersOptions updated, check new fields");
 
-	Options.PlayersToUnregisterCount = 1;
+					Options.PlayersToUnregisterCount = 1;
 
-	// TODO: Do this with ResolveAccountIds instead
-	TArray<EOS_ProductUserId> PlayersToUnregister{ GetProductUserIdChecked(Params.LocalAccountId) };
-	Options.PlayersToUnregister = PlayersToUnregister.GetData();
+					TArray<EOS_ProductUserId> PlayersToUnregister{ GetProductUserIdChecked(Params.LocalAccountId) };
+					Options.PlayersToUnregister = PlayersToUnregister.GetData();
 
-	const FTCHARToUTF8 SessionNameUtf8(*Params.SessionName.ToString());
-	Options.SessionName = SessionNameUtf8.Get();
+					const FTCHARToUTF8 SessionNameUtf8(*Params.SessionName.ToString());
+					Options.SessionName = SessionNameUtf8.Get();
 
-	EOS_Async(EOS_Sessions_UnregisterPlayers, SessionsHandle, Options, 
-		[this, Promise = MoveTemp(Promise), Params](const EOS_Sessions_UnregisterPlayersCallbackInfo* Result) mutable
-		{
-			if (Result->ResultCode != EOS_EResult::EOS_Success)
-			{
-				Promise.EmplaceValue(Errors::FromEOSResult(Result->ResultCode));
-				return;
-			}
-
-			Promise.EmplaceValue(RemoveSessionMemberInternal(Params));
-		});
+					EOS_Async(EOS_Sessions_UnregisterPlayers, SessionsHandle, Options,
+						[this, Promise = MoveTemp(Promise), Params](const EOS_Sessions_UnregisterPlayersCallbackInfo* Result) mutable
+						{
+							if (Result->ResultCode != EOS_EResult::EOS_Success)
+							{
+								Promise.EmplaceValue(Errors::FromEOSResult(Result->ResultCode));
+								return;
+							}
+
+							Promise.EmplaceValue(RemoveSessionMemberInternal(Params));
+						});
+				}
+				else
+				{
+					Promise.EmplaceValue(Errors::InvalidUser());
+				}
+			});
 
 	return Future;
 }
diff --git a/Engine/Plugins/Online/Experimental/OnlineServicesEOSGS/Source/Public/Online/SessionsEOSGS.h b/Engine/Plugins/Online/Experimental/OnlineServicesEOSGS/Source/Public/Online/SessionsEOSGS.h
--- a/Engine/Plugins/Online/Experimental/OnlineServicesEOSGS/Source/Public/Online/SessionsEOSGS.h
+++ b/Engine/Plugins/Online/Experimental/OnlineServicesEOSGS/Source/Public/Online/SessionsEOSGS.h
@@ -207,6 +207,7 @@ protected:
 	void HandleSessionInviteAccepted(const EOS_Sessions_SessionInviteAcceptedCallbackInfo* Data);
 	void HandleJoinSessionAccepted(const EOS_Sessions_JoinSessionAcceptedCallbackInfo* Data);
 
+	void SetHostAddress(EOS_HSessionModification& SessionModHAndle, FString HostAddress);
 	void SetJoinInProgressAllowed(EOS_HSessionModification& SessionModHandle, bool bIsJoinInProgressAllowed);
 	void SetInvitesAllowed(EOS_HSessionModification& SessionModHandle, bool bAreInvitesAllowed);
 	void SetPermissionLevel(EOS_HSessionModification& SessionModHandle, const ESessionJoinPolicy& JoinPolicy);
diff --git a/Engine/Plugins/Online/Experimental/OnlineServicesEOSGS/Source/Public/Online/SessionsEOSGSTypes.h b/Engine/Plugins/Online/Experimental/OnlineServicesEOSGS/Source/Public/Online/SessionsEOSGSTypes.h
--- a/Engine/Plugins/Online/Experimental/OnlineServicesEOSGS/Source/Public/Online/SessionsEOSGSTypes.h
+++ b/Engine/Plugins/Online/Experimental/OnlineServicesEOSGS/Source/Public/Online/SessionsEOSGSTypes.h
@@ -21,6 +21,7 @@ static FName EOSGS_SESSION_ID_OVERRIDE_ATTRIBUTE_KEY = TEXT("EOSGS_SESSION_ID_OV
 static FName EOSGS_OWNER_ACCOUNT_ID_ATTRIBUTE_KEY = TEXT("EOSGS_OWNER_ACCOUNT_ID_ATTRIBUTE_KEY");
 
 static FName EOSGS_BUCKET_ID_ATTRIBUTE_KEY = TEXT("EOSGS_BUCKET_ID_ATTRIBUTE_KEY");
+static FName EOSGS_HOST_ADDRESS_ATTRIBUTE_KEY = TEXT("EOSGS_HOST_ADDRESS_ATTRIBUTE_KEY");
 
 EOS_EOnlineSessionPermissionLevel ToServiceType(const ESessionJoinPolicy& Value);
 ESessionJoinPolicy FromServiceType(const EOS_EOnlineSessionPermissionLevel& Value);
