Removal of SessionIdOverride permanence as custom setting (not needed)

diff --git a/Engine/Plugins/Online/Experimental/OnlineServicesEOSGS/Source/Private/Online/SessionsEOSGS.cpp b/Engine/Plugins/Online/Experimental/OnlineServicesEOSGS/Source/Private/Online/SessionsEOSGS.cpp
--- a/Engine/Plugins/Online/Experimental/OnlineServicesEOSGS/Source/Private/Online/SessionsEOSGS.cpp
+++ b/Engine/Plugins/Online/Experimental/OnlineServicesEOSGS/Source/Private/Online/SessionsEOSGS.cpp
@@ -110,10 +110,6 @@ FSessionEOSGS::FSessionEOSGS(const TSharedPtr<FSessionDetailsHandleEOSGS>& InSes
 				{
 					SessionSettings.SchemaName = FSchemaId(CustomSettingData.Value.GetString());
 				}
-				else if (Key == EOSGS_SESSION_ID_OVERRIDE_ATTRIBUTE_KEY.ToString())
-				{
-					SessionInfo.SessionIdOverride = CustomSettingData.Value.GetString();
-				}
 				else // The rest are parsed as a Custom Session Setting
 				{
 					FCustomSessionSetting CustomSessionSetting;
@@ -580,7 +576,6 @@ void FSessionsEOSGS::WriteCreateSessionModificationHandle(EOS_HSessionModificati
 	SetMaxPlayers(SessionModificationHandle, Params.SessionSettings.NumMaxConnections);
 
 	AddAttribute(SessionModificationHandle, EOSGS_SCHEMA_NAME_ATTRIBUTE_KEY, { FSchemaVariant(Params.SessionSettings.SchemaName.ToString()), ESchemaAttributeVisibility::Public });
-	AddAttribute(SessionModificationHandle, EOSGS_SESSION_ID_OVERRIDE_ATTRIBUTE_KEY, { FSchemaVariant(Params.SessionIdOverride), ESchemaAttributeVisibility::Public });
 
 	EOS_ProductUserId OwnerPUID = GetProductUserId(Params.LocalAccountId);
 	AddAttribute(SessionModificationHandle, EOSGS_OWNER_ACCOUNT_ID_ATTRIBUTE_KEY, { FSchemaVariant(LexToString(OwnerPUID)), ESchemaAttributeVisibility::Public });
diff --git a/Engine/Plugins/Online/Experimental/OnlineServicesEOSGS/Source/Public/Online/SessionsEOSGSTypes.h b/Engine/Plugins/Online/Experimental/OnlineServicesEOSGS/Source/Public/Online/SessionsEOSGSTypes.h
--- a/Engine/Plugins/Online/Experimental/OnlineServicesEOSGS/Source/Public/Online/SessionsEOSGSTypes.h
+++ b/Engine/Plugins/Online/Experimental/OnlineServicesEOSGS/Source/Public/Online/SessionsEOSGSTypes.h
@@ -16,7 +16,6 @@ static FName EOSGS_ALLOW_NEW_MEMBERS_ATTRIBUTE_KEY = TEXT("EOSGS_ALLOW_NEW_MEMBE
 static FName EOSGS_ANTI_CHEAT_PROTECTED_ATTRIBUTE_KEY = TEXT("EOSGS_ANTI_CHEAT_PROTECTED_ATTRIBUTE_KEY");
 static FName EOSGS_IS_DEDICATED_SERVER_SESSION_ATTRIBUTE_KEY = TEXT("EOSGS_IS_DEDICATED_SERVER_SESSION_ATTRIBUTE_KEY");
 static FName EOSGS_SCHEMA_NAME_ATTRIBUTE_KEY = TEXT("EOSGS_SCHEMA_NAME_ATTRIBUTE_KEY");
-static FName EOSGS_SESSION_ID_OVERRIDE_ATTRIBUTE_KEY = TEXT("EOSGS_SESSION_ID_OVERRIDE_ATTRIBUTE_KEY");
 
 static FName EOSGS_OWNER_ACCOUNT_ID_ATTRIBUTE_KEY = TEXT("EOSGS_OWNER_ACCOUNT_ID_ATTRIBUTE_KEY");
 
