Set Presence Session change methods as not implemented in EOS to reflect EOSSDK behaviour

diff --git a/Engine/Plugins/Online/Experimental/OnlineServicesEOSGS/Source/Private/Online/SessionsEOSGS.cpp b/Engine/Plugins/Online/Experimental/OnlineServicesEOSGS/Source/Private/Online/SessionsEOSGS.cpp
--- a/Engine/Plugins/Online/Experimental/OnlineServicesEOSGS/Source/Private/Online/SessionsEOSGS.cpp
+++ b/Engine/Plugins/Online/Experimental/OnlineServicesEOSGS/Source/Private/Online/SessionsEOSGS.cpp
@@ -309,6 +309,16 @@ void FSessionsEOSGS::HandleJoinSessionAccepted(const EOS_Sessions_JoinSessionAcc
 	}
 }
 
+TOnlineResult<FSetPresenceSession> FSessionsEOSGS::SetPresenceSession(FSetPresenceSession::Params&& Params)
+{
+	return TOnlineResult<FSetPresenceSession>(Errors::NotImplemented());
+}
+
+TOnlineResult<FClearPresenceSession> FSessionsEOSGS::ClearPresenceSession(FClearPresenceSession::Params&& Params)
+{
+	return TOnlineResult<FClearPresenceSession>(Errors::NotImplemented());
+}
+
 TFuture<TOnlineResult<FCreateSession>> FSessionsEOSGS::CreateSessionImpl(const FCreateSession::Params& Params)
 {
 	// LAN Sessions
diff --git a/Engine/Plugins/Online/Experimental/OnlineServicesEOSGS/Source/Public/Online/SessionsEOSGS.h b/Engine/Plugins/Online/Experimental/OnlineServicesEOSGS/Source/Public/Online/SessionsEOSGS.h
--- a/Engine/Plugins/Online/Experimental/OnlineServicesEOSGS/Source/Public/Online/SessionsEOSGS.h
+++ b/Engine/Plugins/Online/Experimental/OnlineServicesEOSGS/Source/Public/Online/SessionsEOSGS.h
@@ -190,6 +190,8 @@ public:
 	void Shutdown() override;
 
 	// FSessionsCommon
+	virtual TOnlineResult<FSetPresenceSession> SetPresenceSession(FSetPresenceSession::Params&& Params) override;
+	virtual TOnlineResult<FClearPresenceSession> ClearPresenceSession(FClearPresenceSession::Params&& Params) override;
 	virtual TFuture<TOnlineResult<FCreateSession>> CreateSessionImpl(const FCreateSession::Params& Params) override;
 	virtual TFuture<TOnlineResult<FUpdateSessionSettings>> UpdateSessionSettingsImpl(const FUpdateSessionSettings::Params& Params) override;
 	virtual TFuture<TOnlineResult<FLeaveSession>> LeaveSessionImpl(const FLeaveSession::Params& Params) override;
