OnlineServices - Move additional files out of Experimental

diff --git a/Engine/Platforms/Hololens/Plugins/Online/Experimental/OnlineServicesEOS/OnlineServicesEOS_HoloLens.uplugin b/Engine/Platforms/Hololens/Plugins/Online/Experimental/OnlineServicesEOS/OnlineServicesEOS_HoloLens.uplugin
deleted file mode 100644
diff --git a/Engine/Platforms/Hololens/Plugins/Online/Experimental/OnlineServicesEOS/Source/OnlineServicesEOS_HoloLens.Build.cs b/Engine/Platforms/Hololens/Plugins/Online/Experimental/OnlineServicesEOS/Source/OnlineServicesEOS_HoloLens.Build.cs
deleted file mode 100644
diff --git a/Engine/Platforms/Hololens/Plugins/Online/Experimental/OnlineServicesEOSGS/OnlineServicesEOSGS_HoloLens.uplugin b/Engine/Platforms/Hololens/Plugins/Online/Experimental/OnlineServicesEOSGS/OnlineServicesEOSGS_HoloLens.uplugin
deleted file mode 100644
diff --git a/Engine/Platforms/Hololens/Plugins/Online/Experimental/OnlineServicesEOSGS/Source/OnlineServicesEOSGS_HoloLens.Build.cs b/Engine/Platforms/Hololens/Plugins/Online/Experimental/OnlineServicesEOSGS/Source/OnlineServicesEOSGS_HoloLens.Build.cs
deleted file mode 100644
diff --git a/Engine/Platforms/Hololens/Plugins/Online/OnlineServicesEOS/OnlineServicesEOS_HoloLens.uplugin b/Engine/Platforms/Hololens/Plugins/Online/OnlineServicesEOS/OnlineServicesEOS_HoloLens.uplugin
new file mode 100644
--- /dev/null
+++ b/Engine/Platforms/Hololens/Plugins/Online/OnlineServicesEOS/OnlineServicesEOS_HoloLens.uplugin
@@ -0,0 +1,13 @@
+{
+	"FileVersion": 3,
+	"bIsPluginExtension": true,
+	"Modules": [
+		{
+			"Name": "OnlineServicesEOS",
+			"Type": "Runtime",
+			"PlatformDenyList": [
+				"HoloLens"
+			]
+		}
+	]
+}
\ No newline at end of file
diff --git a/Engine/Platforms/Hololens/Plugins/Online/OnlineServicesEOS/Source/OnlineServicesEOS_HoloLens.Build.cs b/Engine/Platforms/Hololens/Plugins/Online/OnlineServicesEOS/Source/OnlineServicesEOS_HoloLens.Build.cs
new file mode 100644
--- /dev/null
+++ b/Engine/Platforms/Hololens/Plugins/Online/OnlineServicesEOS/Source/OnlineServicesEOS_HoloLens.Build.cs
@@ -0,0 +1,12 @@
+// Copyright Epic Games, Inc. All Rights Reserved.
+
+namespace UnrealBuildTool.Rules
+{
+	public class OnlineServicesEOS_HoloLens : OnlineServicesEOS
+	{
+		public OnlineServicesEOS_HoloLens(ReadOnlyTargetRules Target) : base(Target)
+		{
+			throw new System.Exception("Should not be here - HoloLens in the uplugin's PlatformDenyList");
+		}
+	}
+}
diff --git a/Engine/Platforms/Hololens/Plugins/Online/OnlineServicesEOSGS/OnlineServicesEOSGS_HoloLens.uplugin b/Engine/Platforms/Hololens/Plugins/Online/OnlineServicesEOSGS/OnlineServicesEOSGS_HoloLens.uplugin
new file mode 100644
--- /dev/null
+++ b/Engine/Platforms/Hololens/Plugins/Online/OnlineServicesEOSGS/OnlineServicesEOSGS_HoloLens.uplugin
@@ -0,0 +1,13 @@
+{
+	"FileVersion": 3,
+	"bIsPluginExtension": true,
+	"Modules": [
+		{
+			"Name": "OnlineServicesEOSGS",
+			"Type": "Runtime",
+			"PlatformDenyList": [
+				"HoloLens"
+			]
+		}
+	]
+}
\ No newline at end of file
diff --git a/Engine/Platforms/Hololens/Plugins/Online/OnlineServicesEOSGS/Source/OnlineServicesEOSGS_HoloLens.Build.cs b/Engine/Platforms/Hololens/Plugins/Online/OnlineServicesEOSGS/Source/OnlineServicesEOSGS_HoloLens.Build.cs
new file mode 100644
--- /dev/null
+++ b/Engine/Platforms/Hololens/Plugins/Online/OnlineServicesEOSGS/Source/OnlineServicesEOSGS_HoloLens.Build.cs
@@ -0,0 +1,12 @@
+// Copyright Epic Games, Inc. All Rights Reserved.
+
+namespace UnrealBuildTool.Rules
+{
+	public class OnlineServicesEOSGS_HoloLens : OnlineServicesEOSGS
+	{
+		public OnlineServicesEOSGS_HoloLens(ReadOnlyTargetRules Target) : base(Target)
+		{
+			throw new System.Exception("Should not be here - HoloLens in the uplugin's PlatformDenyList");
+		}
+	}
+}
