OSSv2: Fix the  QueryStats EOS impl crash.

diff --git a/Engine/Plugins/Online/OnlineServicesEOSGS/Source/Private/Online/StatsEOSGS.cpp b/Engine/Plugins/Online/OnlineServicesEOSGS/Source/Private/Online/StatsEOSGS.cpp
--- a/Engine/Plugins/Online/OnlineServicesEOSGS/Source/Private/Online/StatsEOSGS.cpp
+++ b/Engine/Plugins/Online/OnlineServicesEOSGS/Source/Private/Online/StatsEOSGS.cpp
@@ -111,16 +111,20 @@ void QueryStatsEOS(EOS_HStats StatsHandle, const FAccountId& LocalAccountId, con
 	Options.EndTime = EOS_STATS_TIME_UNDEFINED;
 
 	TArray<FStatNameRawBufferEOS> EOSStatNames;
+	TArray<char*> EOSStatNamesPtr;
+
 	EOSStatNames.AddZeroed(StatNames.Num());
+	EOSStatNamesPtr.AddZeroed(StatNames.Num());
 
 	uint32 Index = 0;
 	for (const FString& StatName : StatNames)
 	{
 		FCStringAnsi::Strncpy(EOSStatNames[Index].StatName, TCHAR_TO_UTF8(*StatName), STAT_NAME_MAX_LENGTH_EOS);
+		EOSStatNamesPtr[Index] = EOSStatNames[Index].StatName;
 		Index++;
 	}
 
-	Options.StatNames = (const char**)EOSStatNames.GetData();
+	Options.StatNames = (const char**)EOSStatNamesPtr.GetData();
 	Options.StatNamesCount = StatNames.Num();
 	Options.LocalUserId = GetProductUserIdChecked(LocalAccountId);
 	Options.TargetUserId = GetProductUserIdChecked(TargetAccountId);
@@ -150,7 +154,7 @@ void ReadStatsFromEOSResult(EOS_HStats StatsHandle, const EOS_Stats_OnQueryStats
 		{
 			// Put an empty stat in
 			UE_LOG(LogTemp, VeryVerbose, TEXT("Value not found for stat %s, adding empty value"), *StatName);
-			OutStats.Add(StatName, FStatValue());
+			OutStats.Add(StatName, FStatValue(static_cast<int64>(0)));
 		}
 	}
 }
