[OnlineServicesEOSGS]
* Fixed vs2022 compile errors due to deleted default ctor and then using initializer list

diff --git a/Engine/Plugins/Online/OnlineServicesEOSGS/Source/Private/Online/AuthEOSGS.cpp b/Engine/Plugins/Online/OnlineServicesEOSGS/Source/Private/Online/AuthEOSGS.cpp
--- a/Engine/Plugins/Online/OnlineServicesEOSGS/Source/Private/Online/AuthEOSGS.cpp
+++ b/Engine/Plugins/Online/OnlineServicesEOSGS/Source/Private/Online/AuthEOSGS.cpp
@@ -63,7 +63,7 @@ ENUM_CLASS_FLAGS(EEOSConnectTranslationFlags);
 
 struct FEOSConnectTranslationTraits
 {
-	FEOSConnectTranslationTraits() = delete;
+	//FEOSConnectTranslationTraits() = delete;
 	EOS_EExternalCredentialType Type;
 	EEOSConnectTranslationFlags Flags;
 };
@@ -202,14 +202,12 @@ ENUM_CLASS_FLAGS(EEOSAuthTranslationFlags);
 
 struct FEOSAuthTranslationTraits
 {
-	FEOSAuthTranslationTraits() = delete;
 	EOS_ELoginCredentialType Type;
 	EEOSAuthTranslationFlags Flags;
 };
 
 struct FEOSExternalAuthTranslationTraits
 {
-	FEOSExternalAuthTranslationTraits() = delete;
 	EOS_EExternalCredentialType Type;
 };
 
