Removal of Owning User Id attribute from EOS Sessions and Lobbies to prevent misuse.
Refactored Owning User Id retrievalin EOS Lobbies.
Functionality in Sessions will be re-added when the information is made available via EOSSDK.

diff --git a/Engine/Plugins/Online/OnlineServicesEOSGS/Source/Private/Online/SessionsEOSGS.cpp b/Engine/Plugins/Online/OnlineServicesEOSGS/Source/Private/Online/SessionsEOSGS.cpp
--- a/Engine/Plugins/Online/OnlineServicesEOSGS/Source/Private/Online/SessionsEOSGS.cpp
+++ b/Engine/Plugins/Online/OnlineServicesEOSGS/Source/Private/Online/SessionsEOSGS.cpp
@@ -88,12 +88,6 @@ FSessionEOSGS::FSessionEOSGS(const TSharedPtr<FSessionDetailsHandleEOSGS>& InSes
 				TPair<FSchemaAttributeId, FSchemaVariant> CustomSettingData = CustomSettingConverter.GetAttributeData();
 					
 				// Most Session Settings values get parsed in the same way as Custom Session Settings, so we will attempt to retrieve them
-				if (Key == EOSGS_OWNER_ACCOUNT_ID_ATTRIBUTE_KEY.ToString())
-				{
-					const FString& PlayerIdStr = CustomSettingData.Value.GetString();
-					const EOS_ProductUserId ProductUserId = EOS_ProductUserId_FromString(TCHAR_TO_UTF8(*PlayerIdStr));
-					OwnerAccountId = FindAccountId(ProductUserId);
-				}
 				if (Key == EOSGS_ALLOW_NEW_MEMBERS_ATTRIBUTE_KEY.ToString())
 				{
 					SessionSettings.bAllowNewMembers = CustomSettingData.Value.GetBoolean();
@@ -433,7 +427,6 @@ TFuture<TOnlineResult<FCreateSession>> FSessionsEOSGS::CreateSessionImpl(const F
 
 void FSessionsEOSGS::SetHostAddress(EOS_HSessionModification& SessionModHandle, FString HostAddress)
 {
-	// TODO: We could call EOS_SessionModification_SetHostAddress at this point, although it's not necessary
 	EOS_SessionModification_SetHostAddressOptions Options = {};
 	Options.ApiVersion = EOS_SESSIONMODIFICATION_SETHOSTADDRESS_API_LATEST;
 	static_assert(EOS_SESSIONMODIFICATION_SETHOSTADDRESS_API_LATEST == 1, "EOS_SessionModification_SetHostAddressOptions updated, check new fields");
@@ -587,9 +580,6 @@ void FSessionsEOSGS::WriteCreateSessionModificationHandle(EOS_HSessionModificati
 
 	AddAttribute(SessionModificationHandle, EOSGS_SCHEMA_NAME_ATTRIBUTE_KEY, { FSchemaVariant(Params.SessionSettings.SchemaName.ToString()), ESchemaAttributeVisibility::Public });
 
-	EOS_ProductUserId OwnerPUID = GetProductUserId(Params.LocalAccountId);
-	AddAttribute(SessionModificationHandle, EOSGS_OWNER_ACCOUNT_ID_ATTRIBUTE_KEY, { FSchemaVariant(LexToString(OwnerPUID)), ESchemaAttributeVisibility::Public });
-
 	// Custom Settings
 
  	for (const TPair<FSchemaAttributeId, FCustomSessionSetting>& Entry : Params.SessionSettings.CustomSettings)
@@ -1343,7 +1333,9 @@ TResult<TArray<EOS_ProductUserId>, FOnlineError> GetProductUserIdsFromEOSGSSessi
 {
 	TArray<EOS_ProductUserId> Result;
 
-	EOS_SessionDetails_CopyInfoOptions CopyInfoOptions = { };
+	// This code will remain commented until the related functionality is made available by EOSSDK
+
+	/*EOS_SessionDetails_CopyInfoOptions CopyInfoOptions = {};
 	CopyInfoOptions.ApiVersion = EOS_SESSIONDETAILS_COPYINFO_API_LATEST;
 	static_assert(EOS_SESSIONDETAILS_COPYINFO_API_LATEST == 1, "EOS_SessionDetails_CopyInfoOptions updated, check new fields");
 
@@ -1351,55 +1343,16 @@ TResult<TArray<EOS_ProductUserId>, FOnlineError> GetProductUserIdsFromEOSGSSessi
 	EOS_EResult CopyInfoResult = EOS_SessionDetails_CopyInfo(SessionDetailsHandle, &CopyInfoOptions, &SessionDetailsInfo);
 	if (CopyInfoResult == EOS_EResult::EOS_Success)
 	{
-		// We retrieve all the session attributes
-		EOS_SessionDetails_GetSessionAttributeCountOptions GetAttributeCountOptions = {};
-		GetAttributeCountOptions.ApiVersion = EOS_SESSIONDETAILS_GETSESSIONATTRIBUTECOUNT_API_LATEST;
-		static_assert(EOS_SESSIONDETAILS_GETSESSIONATTRIBUTECOUNT_API_LATEST == 1, "EOS_SessionDetails_GetSessionAttributeCountOptions updated, check new fields");
-		
-		uint32_t AttributeCount = EOS_SessionDetails_GetSessionAttributeCount(SessionDetailsHandle, &GetAttributeCountOptions);
-		for (uint32_t Index = 0; Index < AttributeCount; ++Index)
-		{
-			EOS_SessionDetails_CopySessionAttributeByIndexOptions CopyAttributeByIndexOptions = {};
-			CopyAttributeByIndexOptions.ApiVersion = EOS_SESSIONDETAILS_COPYSESSIONATTRIBUTEBYINDEX_API_LATEST;
-			static_assert(EOS_SESSIONDETAILS_COPYSESSIONATTRIBUTEBYINDEX_API_LATEST == 1, "EOS_SessionDetails_CopySessionAttributeByIndexOptions updated, check new fields");
-
-			CopyAttributeByIndexOptions.AttrIndex = Index;
-
-			EOS_SessionDetails_Attribute* Attribute = nullptr;
-			EOS_EResult CopyAttributeByIndexResult = EOS_SessionDetails_CopySessionAttributeByIndex(SessionDetailsHandle, &CopyAttributeByIndexOptions, &Attribute);
-			if (CopyAttributeByIndexResult == EOS_EResult::EOS_Success)
-			{
-				// We parse a single attribute
-				FString Key(Attribute->Data->Key);
-
-				if (Key == EOSGS_OWNER_ACCOUNT_ID_ATTRIBUTE_KEY.ToString())
-				{
-					FSessionAttributeConverter<ESessionAttributeConversionType::FromService> CustomSettingConverter(*Attribute->Data);
-					TPair<FSchemaAttributeId, FSchemaVariant> CustomSettingData = CustomSettingConverter.GetAttributeData();
-
-					const FString& PlayerIdStr = CustomSettingData.Value.GetString();
-					const EOS_ProductUserId ProductUserId = EOS_ProductUserId_FromString(TCHAR_TO_UTF8(*PlayerIdStr));
-					Result.AddUnique(ProductUserId);
-
-					break;
-				}
-			}
-			else
-			{
-				UE_LOG(LogTemp, Warning, TEXT("[FSessionsEOSGS::BuildSessionFromDetailsHandle] EOS_SessionDetails_CopySessionAttributeByIndex failed with result [%s]"), *LexToString(CopyAttributeByIndexResult));
-
-				TResult<TArray<EOS_ProductUserId>, FOnlineError>(Errors::FromEOSResult(CopyAttributeByIndexResult));
-			}
-		}
+		// TODO: When available, we'll retrieve the Session's Owner Id here.
 	}
 	else
 	{
 		UE_LOG(LogTemp, Warning, TEXT("[FSessionsEOSGS::BuildSessionFromDetailsHandle] EOS_SessionDetails_CopyInfo failed with result [%s]"), *LexToString(CopyInfoResult));
 
 		TResult<TArray<EOS_ProductUserId>, FOnlineError>(Errors::FromEOSResult(CopyInfoResult));
-	}
+	}*/
 
-	return TResult<TArray<EOS_ProductUserId>, FOnlineError>(Result);
+	return TResult<TArray<EOS_ProductUserId>, FOnlineError>(MoveTemp(Result));
 }
 
 TOnlineAsyncOpHandle<FBuildSessionFromDetailsHandle> FSessionsEOSGS::BuildSessionFromDetailsHandle(FBuildSessionFromDetailsHandle::Params&& Params)
diff --git a/Engine/Plugins/Online/OnlineServicesEOSGS/Source/Public/Online/SessionsEOSGSTypes.h b/Engine/Plugins/Online/OnlineServicesEOSGS/Source/Public/Online/SessionsEOSGSTypes.h
--- a/Engine/Plugins/Online/OnlineServicesEOSGS/Source/Public/Online/SessionsEOSGSTypes.h
+++ b/Engine/Plugins/Online/OnlineServicesEOSGS/Source/Public/Online/SessionsEOSGSTypes.h
@@ -17,8 +17,6 @@ static FName EOSGS_ANTI_CHEAT_PROTECTED_ATTRIBUTE_KEY = TEXT("EOSGS_ANTI_CHEAT_P
 static FName EOSGS_IS_DEDICATED_SERVER_SESSION_ATTRIBUTE_KEY = TEXT("EOSGS_IS_DEDICATED_SERVER_SESSION_ATTRIBUTE_KEY");
 static FName EOSGS_SCHEMA_NAME_ATTRIBUTE_KEY = TEXT("EOSGS_SCHEMA_NAME_ATTRIBUTE_KEY");
 
-static FName EOSGS_OWNER_ACCOUNT_ID_ATTRIBUTE_KEY = TEXT("EOSGS_OWNER_ACCOUNT_ID_ATTRIBUTE_KEY");
-
 static FName EOSGS_BUCKET_ID_ATTRIBUTE_KEY = TEXT("EOSGS_BUCKET_ID_ATTRIBUTE_KEY");
 static FName EOSGS_HOST_ADDRESS_ATTRIBUTE_KEY = TEXT("EOSGS_HOST_ADDRESS_ATTRIBUTE_KEY");
 
diff --git a/Engine/Plugins/Online/OnlineSubsystemEOS/Source/OnlineSubsystemEOS/Private/OnlineSessionEOS.cpp b/Engine/Plugins/Online/OnlineSubsystemEOS/Source/OnlineSubsystemEOS/Private/OnlineSessionEOS.cpp
--- a/Engine/Plugins/Online/OnlineSubsystemEOS/Source/OnlineSubsystemEOS/Private/OnlineSessionEOS.cpp
+++ b/Engine/Plugins/Online/OnlineSubsystemEOS/Source/OnlineSubsystemEOS/Private/OnlineSessionEOS.cpp
@@ -1155,33 +1155,6 @@ void FOnlineSessionEOS::SetAttributes(EOS_HSessionModification SessionModHandle,
 	FAttributeOptions Opt2("NumPublicConnections", Session->SessionSettings.NumPublicConnections);
 	AddAttribute(SessionModHandle, &Opt2);
 
-	if (Session->OwningUserId->IsValid())
-	{
-		FAttributeOptions OwningUserId("OwningUserId", TCHAR_TO_UTF8(*Session->OwningUserId->ToString()));
-		AddAttribute(SessionModHandle, &OwningUserId);
-	}
-
-	// Handle auto generation of dedicated server names
-	if (Session->OwningUserName.IsEmpty())
-	{
-		FString OwningPlayerName(TEXT("DedicatedServer - "));
-
-		FString UserName = FPlatformProcess::UserName();
-		if (UserName.IsEmpty())
-		{
-			FString ComputerName = FPlatformProcess::ComputerName();
-			OwningPlayerName += ComputerName;
-		}
-		else
-		{
-			OwningPlayerName += UserName;
-		}
-		Session->OwningUserName = OwningPlayerName;
-	}
-
-	FAttributeOptions OwningUserName("OwningUserName", TCHAR_TO_UTF8(*Session->OwningUserName));
-	AddAttribute(SessionModHandle, &OwningUserName);
-
 	FAttributeOptions Opt5("bAntiCheatProtected", Session->SessionSettings.bAntiCheatProtected);
 	AddAttribute(SessionModHandle, &Opt5);
 
@@ -2095,14 +2068,6 @@ void FOnlineSessionEOS::CopyAttributes(EOS_HSessionDetails SessionHandle, FOnlin
 				// Adjust the private connections based upon this
 				OutSession.SessionSettings.NumPrivateConnections = Attribute->Data->Value.AsInt64;
 			}
-			else if (Key == TEXT("OwningUserId"))
-			{
-				OutSession.OwningUserId = FUniqueNetIdEOSRegistry::FindOrAdd(UTF8_TO_TCHAR(Attribute->Data->Value.AsUtf8));
-			}
-			else if (Key == TEXT("OwningUserName"))
-			{
-				OutSession.OwningUserName = UTF8_TO_TCHAR(Attribute->Data->Value.AsUtf8);
-			}
 			else if (Key == TEXT("bAntiCheatProtected"))
 			{
 				OutSession.SessionSettings.bAntiCheatProtected = Attribute->Data->Value.AsBool == EOS_TRUE;
@@ -3816,13 +3781,6 @@ void FOnlineSessionEOS::SetLobbyAttributes(EOS_HLobbyModification LobbyModificat
 	const FLobbyAttributeOptions SearchLobbiesAttribute(TCHAR_TO_UTF8(*SearchLobbies), true);
 	AddLobbyAttribute(LobbyModificationHandle, &SearchLobbiesAttribute);
 
-	// We set the session's owner id and name
-	const FLobbyAttributeOptions OwnerId("OwningUserId", TCHAR_TO_UTF8(*Session->OwningUserId->ToString()));
-	AddLobbyAttribute(LobbyModificationHandle, &OwnerId);
-
-	const FLobbyAttributeOptions OwnerName("OwningUserName", TCHAR_TO_UTF8(*Session->OwningUserName));
-	AddLobbyAttribute(LobbyModificationHandle, &OwnerName);
-
 	// Now the session settings
 	const FLobbyAttributeOptions Opt1("NumPrivateConnections", Session->SessionSettings.NumPrivateConnections);
 	AddLobbyAttribute(LobbyModificationHandle, &Opt1);
@@ -4276,6 +4234,19 @@ void FOnlineSessionEOS::CopyLobbyData(const TSharedRef<FLobbyDetailsEOS>& LobbyD
 				FOnlineSession* Session = GetOnlineSessionFromLobbyId(*LobbyId);
 				if (Session)
 				{
+					// One of the resolved ids will be the Owner's, so we'll set that too
+					EOS_LobbyDetails_GetLobbyOwnerOptions GetLobbyOwnerOptions = {};
+					GetLobbyOwnerOptions.ApiVersion = EOS_LOBBYDETAILS_GETLOBBYOWNER_API_LATEST;
+					static_assert(EOS_LOBBYDETAILS_GETLOBBYOWNER_API_LATEST == 1, "EOS_LobbyDetails_GetLobbyOwnerOptions updated, check new fields");
+
+					const EOS_ProductUserId LobbyOwner = EOS_LobbyDetails_GetLobbyOwner(LobbyDetails->LobbyDetailsHandle, &GetLobbyOwnerOptions);
+
+					if (FUniqueNetIdEOSRef* OwnerNetId = ResolvedUniqueNetIds.Find(LobbyOwner))
+					{
+						Session->OwningUserId = *OwnerNetId;
+						Session->OwningUserName = EOSSubsystem->UserManager->GetPlayerNickname(**OwnerNetId);
+					}
+
 					for (TMap<EOS_ProductUserId, FUniqueNetIdEOSRef>::TConstIterator It(ResolvedUniqueNetIds); It; ++It)
 					{
 						FSessionSettings& MemberSettings = Session->SessionSettings.MemberSettings.FindOrAdd(It.Value());
@@ -4313,15 +4284,7 @@ void FOnlineSessionEOS::CopyLobbyAttributes(const TSharedRef<FLobbyDetailsEOS>&
 		if (ResultCode == EOS_EResult::EOS_Success)
 		{
 			FString Key = UTF8_TO_TCHAR(Attribute->Data->Key);
-			if (Key == TEXT("OwningUserId"))
-			{
-				OutSession.OwningUserId = FUniqueNetIdEOSRegistry::FindOrAdd(UTF8_TO_TCHAR(Attribute->Data->Value.AsUtf8));
-			}
-			else if (Key == TEXT("OwningUserName"))
-			{
-				OutSession.OwningUserName = UTF8_TO_TCHAR(Attribute->Data->Value.AsUtf8);
-			}
-			else if (Key == TEXT("NumPublicConnections"))
+			if (Key == TEXT("NumPublicConnections"))
 			{
 				OutSession.SessionSettings.NumPublicConnections = Attribute->Data->Value.AsInt64;
 			}
