PR #9742: Fix PersistentAuth (Reported by [REDACTED_PERSON])

diff --git a/Engine/Plugins/Online/OnlineServicesEOSGS/Source/Private/Online/AuthEOSGS.cpp b/Engine/Plugins/Online/OnlineServicesEOSGS/Source/Private/Online/AuthEOSGS.cpp
--- a/Engine/Plugins/Online/OnlineServicesEOSGS/Source/Private/Online/AuthEOSGS.cpp
+++ b/Engine/Plugins/Online/OnlineServicesEOSGS/Source/Private/Online/AuthEOSGS.cpp
@@ -247,7 +247,7 @@ const FEOSAuthTranslationTraits* FEOSAuthLoginOptions::GetLoginTranslatorTraits(
 	static const TMap<FName, FEOSAuthTranslationTraits> SupportedLoginTranslatorTraits = {
 		{ LoginCredentialsType::Password, { EOS_ELoginCredentialType::EOS_LCT_Password, EEOSAuthTranslationFlags::SetId | EEOSAuthTranslationFlags::SetTokenFromString } },
 		{ LoginCredentialsType::ExchangeCode, { EOS_ELoginCredentialType::EOS_LCT_ExchangeCode, EEOSAuthTranslationFlags::SetTokenFromString } },
-		{ LoginCredentialsType::PersistentAuth, { EOS_ELoginCredentialType::EOS_LCT_PersistentAuth, EEOSAuthTranslationFlags::SetTokenFromString } },
+		{ LoginCredentialsType::PersistentAuth, { EOS_ELoginCredentialType::EOS_LCT_PersistentAuth, EEOSAuthTranslationFlags::None } },
 		{ LoginCredentialsType::Developer, { EOS_ELoginCredentialType::EOS_LCT_Developer, EEOSAuthTranslationFlags::SetId | EEOSAuthTranslationFlags::SetTokenFromString } },
 		{ LoginCredentialsType::RefreshToken, { EOS_ELoginCredentialType::EOS_LCT_RefreshToken, EEOSAuthTranslationFlags::SetTokenFromString } },
 		{ LoginCredentialsType::AccountPortal, { EOS_ELoginCredentialType::EOS_LCT_AccountPortal, EEOSAuthTranslationFlags::SetId | EEOSAuthTranslationFlags::SetTokenFromString } },
