Fix the OSSv2 EOS tests for Stats.

diff --git a/Engine/Plugins/Online/OnlineServicesEOSGS/Source/Private/Online/StatsEOSGS.cpp b/Engine/Plugins/Online/OnlineServicesEOSGS/Source/Private/Online/StatsEOSGS.cpp
--- a/Engine/Plugins/Online/OnlineServicesEOSGS/Source/Private/Online/StatsEOSGS.cpp
+++ b/Engine/Plugins/Online/OnlineServicesEOSGS/Source/Private/Online/StatsEOSGS.cpp
@@ -41,7 +41,15 @@ TOnlineAsyncOpHandle<FUpdateStats> FStatsEOSGS::UpdateStats(FUpdateStats::Params
 		return Op->GetHandle();
 	}
 
-	// TODO: Only load stats definitions on platforms which requires to upload leaderboard score separately
+	for (const FUserStats& UpdateUserStats : Op->GetParams().UpdateUsersStats)
+	{
+		if (Op->GetParams().LocalAccountId != UpdateUserStats.AccountId)
+		{
+			// Client can only update the stats of the local account itself
+			Op->SetError(Errors::InvalidParams());
+			return Op->GetHandle();
+		}
+	}
 
 	for (const FUserStats& UpdateUserStats : Op->GetParams().UpdateUsersStats)
 	{
@@ -230,7 +238,7 @@ TOnlineAsyncOpHandle<FBatchQueryStats> FStatsEOSGS::BatchQueryStats(FBatchQueryS
 			{
 				Private::QueryStatsEOS(StatsHandle, InAsyncOp.GetParams().LocalAccountId, TargetAccountId, InAsyncOp.GetParams().StatNames, MoveTemp(Promise));
 			})
-			.Then([this](TOnlineAsyncOp<FBatchQueryStats>& InAsyncOp, const EOS_Stats_OnQueryStatsCompleteCallbackInfo* Data)
+			.Then([this, TargetAccountId](TOnlineAsyncOp<FBatchQueryStats>& InAsyncOp, const EOS_Stats_OnQueryStatsCompleteCallbackInfo* Data)
 			{
 				if (Data->ResultCode != EOS_EResult::EOS_Success)
 				{
@@ -241,6 +249,7 @@ TOnlineAsyncOpHandle<FBatchQueryStats> FStatsEOSGS::BatchQueryStats(FBatchQueryS
 				}
 
 				FUserStats& UserStats = BatchQueriedUsersStats.Emplace_GetRef();
+				UserStats.AccountId = TargetAccountId;
 
 				Private::ReadStatsFromEOSResult(StatsHandle, Data, InAsyncOp.GetParams().StatNames, UserStats.Stats);
 				CacheUserStats(UserStats);
