Fix remain untagged memory allocations in online, by replacing local defined global variables with predefined online error code, or use raw text instead of global string.

Reason: global variables which contains string will allocate memory when start up before engine memory allocation scope.

diff --git a/Engine/Plugins/Online/OnlineServices/Source/OnlineServicesCommon/Public/Online/OnlineExecHandler.h b/Engine/Plugins/Online/OnlineServices/Source/OnlineServicesCommon/Public/Online/OnlineExecHandler.h
--- a/Engine/Plugins/Online/OnlineServices/Source/OnlineServicesCommon/Public/Online/OnlineExecHandler.h
+++ b/Engine/Plugins/Online/OnlineServices/Source/OnlineServicesCommon/Public/Online/OnlineExecHandler.h
@@ -132,7 +132,7 @@ struct TOnlineVariantVisitInfo
 	//static_assert(sizeof(T) == -1, "If you are hitting this, you are using a TVariant with a value that does not have a TOnlineVariantVisitInfo! Make sure you are adding a TOnlineVariantVisitInfo for your class! (Check LastType or FirstType below for the type that needs attention)");
 
 public:
-	inline static TArray<FString> Prefixes = { TEXT("[GARBAGE]") };
+	inline static const TCHAR* Prefixes[] = {TEXT("[GARBAGE]")};
 
 	static bool Assign(FString& VariantValue, VariantObject& Variant, IOnlineServices* Services)
 	{
@@ -144,7 +144,7 @@ template< typename VariantObject >
 struct TOnlineVariantVisitInfo<int64, VariantObject>
 {
 public:
-	inline static TArray<FString> Prefixes = { TEXT("i"),TEXT("i64"),TEXT("int"),TEXT("int64") };
+	inline static const TCHAR* Prefixes[] = { TEXT("i"),TEXT("i64"),TEXT("int"),TEXT("int64") };
 
 	static bool Assign(FString& VariantValue, VariantObject& Variant, IOnlineServices* Services)
 	{
@@ -157,7 +157,7 @@ template< typename VariantObject >
 struct TOnlineVariantVisitInfo<int32, VariantObject>
 {
 public:
-	inline static TArray<FString> Prefixes = { TEXT("i32"),TEXT("int32") };
+	inline static const TCHAR* Prefixes[] = { TEXT("i32"),TEXT("int32") };
 
 	static bool Assign(FString& VariantValue, VariantObject& Variant, IOnlineServices* Services)
 	{
@@ -171,7 +171,7 @@ template< typename VariantObject >
 struct TOnlineVariantVisitInfo<double, VariantObject>
 {
 public:
-	inline static TArray<FString> Prefixes = { TEXT("d"),TEXT("double") };
+	inline static const TCHAR* Prefixes[] = { TEXT("d"),TEXT("double") };
 
 	static bool Assign(FString& VariantValue, VariantObject& Variant, IOnlineServices* Services)
 	{
@@ -184,7 +184,7 @@ template< typename VariantObject >
 struct TOnlineVariantVisitInfo<float, VariantObject>
 {
 public:
-	inline static TArray<FString> Prefixes = { TEXT("f"),TEXT("float") };
+	inline static const TCHAR* Prefixes[] = { TEXT("f"),TEXT("float") };
 
 	static bool Assign(FString& VariantValue, VariantObject& Variant, IOnlineServices* Services)
 	{
@@ -197,7 +197,7 @@ template< typename VariantObject >
 struct TOnlineVariantVisitInfo<bool, VariantObject>
 {
 public:
-	inline static TArray<FString> Prefixes = { TEXT("b"),TEXT("bool"),TEXT("boolean") };
+	inline static const TCHAR* Prefixes[] = { TEXT("b"),TEXT("bool"),TEXT("boolean") };
 
 	static bool Assign(FString& VariantValue, VariantObject& Variant, IOnlineServices* Services)
 	{
@@ -211,7 +211,7 @@ template< typename VariantObject >
 struct TOnlineVariantVisitInfo<FString, VariantObject>
 {
 public:
-	inline static TArray<FString> Prefixes = { TEXT("s"),TEXT("str"),TEXT("string") };
+	inline static const TCHAR* Prefixes[] = { TEXT("s"),TEXT("str"),TEXT("string") };
 
 	static bool Assign(FString& VariantValue, VariantObject& Variant, IOnlineServices* Services)
 	{
@@ -224,7 +224,7 @@ template< typename VariantObject >
 struct TOnlineVariantVisitInfo<FAccountId, VariantObject>
 {
 public:
-	inline static TArray<FString> Prefixes = { TEXT("u"),TEXT("user") };
+	inline static const TCHAR* Prefixes[] = {TEXT("u"),TEXT("user")};
 
 	static bool Assign(FString& VariantValue, VariantObject& Variant, IOnlineServices* Services)
 	{
@@ -244,7 +244,7 @@ public:
 struct UE::Online::Private::TOnlineVariantVisitInfo<EnumName, VariantObject>\
 {\
 public:\
-	inline static TArray<FString> Prefixes = {TEXT(#ShortName),TEXT("e")};\
+	inline static const TCHAR* Prefixes[] = {TEXT(#ShortName),TEXT("e")};\
 	static bool Assign(FString& VariantValue, VariantObject& Variant, IOnlineServices* Services)\
 	{\
 		EnumName Status;\
@@ -719,9 +719,9 @@ inline bool ParseOnlineExecParams(const TCHAR*& Cmd, TOptional<T>& Optional, IOn
 template<typename VariantType, typename LastType>
 static bool VisitVariantTypes(VariantType& Variant, FString& VariantPrefix, FString& VariantValue, IOnlineServices* Services)
 {
-	for(const FString& Str : TOnlineVariantVisitInfo<LastType, VariantType>::Prefixes)
+	for(const TCHAR* Prefix : TOnlineVariantVisitInfo<LastType, VariantType>::Prefixes)
 	{
-		if (Str.Equals(VariantPrefix, ESearchCase::IgnoreCase))
+		if (!VariantPrefix.IsEmpty() && FCString::Stricmp(Prefix, *VariantPrefix) == 0)
 		{
 			return TOnlineVariantVisitInfo<LastType, VariantType>::Assign(VariantValue, Variant, Services);
 		}
@@ -732,9 +732,9 @@ static bool VisitVariantTypes(VariantType& Variant, FString& VariantPrefix, FStr
 template<typename VariantType, typename FirstType, typename SecondType, typename... RemainingTypes>
 inline bool VisitVariantTypes(VariantType& Variant, FString& VariantPrefix, FString& VariantValue, IOnlineServices* Services)
 {
-	for(const FString& Str : TOnlineVariantVisitInfo<FirstType, VariantType>::Prefixes)
+	for(const TCHAR* Prefix : TOnlineVariantVisitInfo<FirstType, VariantType>::Prefixes)
 	{
-		if (Str.Equals(VariantPrefix, ESearchCase::IgnoreCase))
+		if (!VariantPrefix.IsEmpty() && FCString::Stricmp(Prefix, *VariantPrefix) == 0)
 		{
 			return TOnlineVariantVisitInfo<FirstType, VariantType>::Assign(VariantValue, Variant, Services);
 		}
diff --git a/Engine/Plugins/Online/OnlineServicesEOS/Source/Private/Online/AuthEOS.cpp b/Engine/Plugins/Online/OnlineServicesEOS/Source/Private/Online/AuthEOS.cpp
--- a/Engine/Plugins/Online/OnlineServicesEOS/Source/Private/Online/AuthEOS.cpp
+++ b/Engine/Plugins/Online/OnlineServicesEOS/Source/Private/Online/AuthEOS.cpp
@@ -33,7 +33,7 @@ END_ONLINE_STRUCT_META()
 namespace
 {
 
-static const FString AccountInfoKeyName = TEXT("AccountInfoEOS");
+#define UE_ONLINE_AUTH_EOS_ACCOUNT_INFO_KEY_NAME TEXT("AccountInfoEOS")
 
 /* anonymous */ }
 
@@ -69,7 +69,7 @@ TOnlineAsyncOpHandle<FAuthLogin> FAuthEOS::Login(FAuthLogin::Params&& Params)
 		AccountInfoEOS->LoginStatus = ELoginStatus::NotLoggedIn;
 
 		// Set user auth data on operation.
-		InAsyncOp.Data.Set<TSharedRef<FAccountInfoEOS>>(AccountInfoKeyName, AccountInfoEOS.ToSharedRef());
+		InAsyncOp.Data.Set<TSharedRef<FAccountInfoEOS>>(UE_ONLINE_AUTH_EOS_ACCOUNT_INFO_KEY_NAME, AccountInfoEOS.ToSharedRef());
 	})
 	// Step 2: Login EAS.
 	.Then([this](TOnlineAsyncOp<FAuthLogin>& InAsyncOp)
@@ -94,7 +94,7 @@ TOnlineAsyncOpHandle<FAuthLogin> FAuthEOS::Login(FAuthLogin::Params&& Params)
 		{
 			if (TSharedPtr<TOnlineAsyncOp<FAuthLogin>> Op = WeakOp.Pin())
 			{
-				const TSharedRef<FAccountInfoEOS>& AccountInfoEOS = GetOpDataChecked<TSharedRef<FAccountInfoEOS>>(*Op, AccountInfoKeyName);
+				const TSharedRef<FAccountInfoEOS>& AccountInfoEOS = GetOpDataChecked<TSharedRef<FAccountInfoEOS>>(*Op, UE_ONLINE_AUTH_EOS_ACCOUNT_INFO_KEY_NAME);
 
 				if (LoginResult.IsError())
 				{
@@ -117,7 +117,7 @@ TOnlineAsyncOpHandle<FAuthLogin> FAuthEOS::Login(FAuthLogin::Params&& Params)
 	.Then([this](TOnlineAsyncOp<FAuthLogin>& InAsyncOp)
 	{
 		const FAuthLogin::Params& Params = InAsyncOp.GetParams();
-		const TSharedRef<FAccountInfoEOS>& AccountInfoEOS = GetOpDataChecked<TSharedRef<FAccountInfoEOS>>(InAsyncOp, AccountInfoKeyName);
+		const TSharedRef<FAccountInfoEOS>& AccountInfoEOS = GetOpDataChecked<TSharedRef<FAccountInfoEOS>>(InAsyncOp, UE_ONLINE_AUTH_EOS_ACCOUNT_INFO_KEY_NAME);
 
 		TPromise<FAuthLoginConnectImpl::Params> Promise;
 		TFuture<FAuthLoginConnectImpl::Params> Future = Promise.GetFuture();
@@ -148,7 +148,7 @@ TOnlineAsyncOpHandle<FAuthLogin> FAuthEOS::Login(FAuthLogin::Params&& Params)
 	// Step 4: Attempt connect login. On connect login failure handle logout of EAS.
 	.Then([this](TOnlineAsyncOp<FAuthLogin>& InAsyncOp, FAuthLoginConnectImpl::Params&& LoginConnectParams)
 	{
-		const TSharedRef<FAccountInfoEOS>& AccountInfoEOS = GetOpDataChecked<TSharedRef<FAccountInfoEOS>>(InAsyncOp, AccountInfoKeyName);
+		const TSharedRef<FAccountInfoEOS>& AccountInfoEOS = GetOpDataChecked<TSharedRef<FAccountInfoEOS>>(InAsyncOp, UE_ONLINE_AUTH_EOS_ACCOUNT_INFO_KEY_NAME);
 
 		TPromise<void> Promise;
 		TFuture<void> Future = Promise.GetFuture();
@@ -192,7 +192,7 @@ TOnlineAsyncOpHandle<FAuthLogin> FAuthEOS::Login(FAuthLogin::Params&& Params)
 	// Step 5: Fetch dependent data.
 	.Then([this](TOnlineAsyncOp<FAuthLogin>& InAsyncOp)
 	{
-		const TSharedRef<FAccountInfoEOS>& AccountInfoEOS = GetOpDataChecked<TSharedRef<FAccountInfoEOS>>(InAsyncOp, AccountInfoKeyName);
+		const TSharedRef<FAccountInfoEOS>& AccountInfoEOS = GetOpDataChecked<TSharedRef<FAccountInfoEOS>>(InAsyncOp, UE_ONLINE_AUTH_EOS_ACCOUNT_INFO_KEY_NAME);
 
 		// Get display name
 		EOS_UserInfo_CopyUserInfoOptions Options = { };
@@ -236,7 +236,7 @@ TOnlineAsyncOpHandle<FAuthLogin> FAuthEOS::Login(FAuthLogin::Params&& Params)
 	// Step 6: bookkeeping and notifications.
 	.Then([this](TOnlineAsyncOp<FAuthLogin>& InAsyncOp)
 	{
-		const TSharedRef<FAccountInfoEOS>& AccountInfoEOS = GetOpDataChecked<TSharedRef<FAccountInfoEOS>>(InAsyncOp, AccountInfoKeyName);
+		const TSharedRef<FAccountInfoEOS>& AccountInfoEOS = GetOpDataChecked<TSharedRef<FAccountInfoEOS>>(InAsyncOp, UE_ONLINE_AUTH_EOS_ACCOUNT_INFO_KEY_NAME);
 		AccountInfoEOS->LoginStatus = ELoginStatus::LoggedIn;
 		AccountInfoEOS->AccountId = CreateAccountId(AccountInfoEOS->EpicAccountId, AccountInfoEOS->ProductUserId);
 		AccountInfoRegistryEOS.Register(AccountInfoEOS);
diff --git a/Engine/Plugins/Online/OnlineServicesEOSGS/Source/Private/Online/LobbiesEOSGS.cpp b/Engine/Plugins/Online/OnlineServicesEOSGS/Source/Private/Online/LobbiesEOSGS.cpp
--- a/Engine/Plugins/Online/OnlineServicesEOSGS/Source/Private/Online/LobbiesEOSGS.cpp
+++ b/Engine/Plugins/Online/OnlineServicesEOSGS/Source/Private/Online/LobbiesEOSGS.cpp
@@ -14,12 +14,12 @@
 #include "Online/OnlineUtils.h"
 #include "eos_lobby.h"
 
-namespace UE::Online {
+#define UE_ONLINE_LOBBY_EOS_KEY_NAME_LOBBY_DATA TEXT("LobbyData")
+#define UE_ONLINE_LOBBY_EOS_KEY_NAME_LOBBY_DETAILS TEXT("LobbyDetails")
+#define UE_ONLINE_LOBBY_EOS_KEY_NAME_LOBBY_ID_STRING TEXT("LobbyIdString")
+#define UE_ONLINE_LOBBY_EOS_KEY_NAME_LOBBY_SEARCH TEXT("LobbySearch")
 
-static const FString LobbyDataKeyName = TEXT("LobbyData");
-static const FString LobbyDetailsKeyName = TEXT("LobbyDetails");
-static const FString LobbyIdStringKeyName = TEXT("LobbyIdString");
-static const FString LobbySearchKeyName = TEXT("LobbySearch");
+namespace UE::Online {
 
 static const int32 MaxAttributeSize = 1000;
 
@@ -155,14 +155,14 @@ TOnlineAsyncOpHandle<FCreateLobby> FLobbiesEOSGS::CreateLobby(FCreateLobby::Para
 		else
 		{
 			// Store the lobby id string on the operation data. This is used to destroy the lobby when creating the lobby structures fails.
-			InAsyncOp.Data.Set<FString>(LobbyIdStringKeyName, FString(UTF8_TO_TCHAR(Data->LobbyId)));
+			InAsyncOp.Data.Set<FString>(UE_ONLINE_LOBBY_EOS_KEY_NAME_LOBBY_ID_STRING, FString(UTF8_TO_TCHAR(Data->LobbyId)));
 		}
 	})
 	// Step 3: Create the lobby details object from the lobby id.
 	.Then([this, DestroyLobbyDuringCreate](TOnlineAsyncOp<FCreateLobby>& InAsyncOp)
 	{
 		const FCreateLobby::Params& Params = InAsyncOp.GetParams();
-		const FString& LobbyIdString = GetOpDataChecked<FString>(InAsyncOp, LobbyIdStringKeyName);
+		const FString& LobbyIdString = GetOpDataChecked<FString>(InAsyncOp, UE_ONLINE_LOBBY_EOS_KEY_NAME_LOBBY_ID_STRING);
 
 		// Try to create the lobby details object.
 		TDefaultErrorResultInternal<TSharedRef<FLobbyDetailsEOS>> Result =
@@ -177,21 +177,21 @@ TOnlineAsyncOpHandle<FCreateLobby> FLobbiesEOSGS::CreateLobby(FCreateLobby::Para
 		else
 		{
 			// Attach the lobby details to the operation.
-			InAsyncOp.Data.Set<TSharedRef<FLobbyDetailsEOS>>(LobbyDetailsKeyName, Result.GetOkValue());
+			InAsyncOp.Data.Set<TSharedRef<FLobbyDetailsEOS>>(UE_ONLINE_LOBBY_EOS_KEY_NAME_LOBBY_DETAILS, Result.GetOkValue());
 			return MakeFulfilledPromise<void>().GetFuture();
 		}
 	})
 	// Step 4: Create the lobby data object from the lobby details.
 	.Then([this](TOnlineAsyncOp<FCreateLobby>& InAsyncOp)
 	{
-		TSharedRef<FLobbyDetailsEOS> LobbyDetails = GetOpDataChecked<TSharedRef<FLobbyDetailsEOS>>(InAsyncOp, LobbyDetailsKeyName);
+		TSharedRef<FLobbyDetailsEOS> LobbyDetails = GetOpDataChecked<TSharedRef<FLobbyDetailsEOS>>(InAsyncOp, UE_ONLINE_LOBBY_EOS_KEY_NAME_LOBBY_DETAILS);
 		return LobbyDataRegistry->FindOrCreateFromLobbyDetails(InAsyncOp.GetParams().LocalAccountId, LobbyDetails);
 	})
 	// Step 5: Handle errors and store the lobby data on the async op properties.
 	.Then([this, DestroyLobbyDuringCreate](TOnlineAsyncOp<FCreateLobby>& InAsyncOp, TDefaultErrorResultInternal<TSharedRef<FLobbyDataEOS>>&& Result)
 	{
 		const FCreateLobby::Params& Params = InAsyncOp.GetParams();
-		const FString& LobbyIdString = GetOpDataChecked<FString>(InAsyncOp, LobbyIdStringKeyName);
+		const FString& LobbyIdString = GetOpDataChecked<FString>(InAsyncOp, UE_ONLINE_LOBBY_EOS_KEY_NAME_LOBBY_ID_STRING);
 
 		if (Result.IsError())
 		{
@@ -203,7 +203,7 @@ TOnlineAsyncOpHandle<FCreateLobby> FLobbiesEOSGS::CreateLobby(FCreateLobby::Para
 		else
 		{
 			// Store lobby data on the operation.
-			InAsyncOp.Data.Set<TSharedRef<FLobbyDataEOS>>(LobbyDataKeyName, MoveTemp(Result.GetOkValue()));
+			InAsyncOp.Data.Set<TSharedRef<FLobbyDataEOS>>(UE_ONLINE_LOBBY_EOS_KEY_NAME_LOBBY_DATA, MoveTemp(Result.GetOkValue()));
 			return MakeFulfilledPromise<void>().GetFuture();
 		}
 	})
@@ -211,7 +211,7 @@ TOnlineAsyncOpHandle<FCreateLobby> FLobbiesEOSGS::CreateLobby(FCreateLobby::Para
 	.Then([this](TOnlineAsyncOp<FCreateLobby>& InAsyncOp)
 	{
 		const FCreateLobby::Params& Params = InAsyncOp.GetParams();
-		const TSharedRef<FLobbyDataEOS>& LobbyData = GetOpDataChecked<TSharedRef<FLobbyDataEOS>>(InAsyncOp, LobbyDataKeyName);
+		const TSharedRef<FLobbyDataEOS>& LobbyData = GetOpDataChecked<TSharedRef<FLobbyDataEOS>>(InAsyncOp, UE_ONLINE_LOBBY_EOS_KEY_NAME_LOBBY_DATA);
 
 		FLobbyClientDataPrepareClientChanges::Params PrepareParams;
 		PrepareParams.LocalAccountId = Params.LocalAccountId;
@@ -239,7 +239,7 @@ TOnlineAsyncOpHandle<FCreateLobby> FLobbiesEOSGS::CreateLobby(FCreateLobby::Para
 	.Then([this, DestroyLobbyDuringCreate](TOnlineAsyncOp<FCreateLobby>& InAsyncOp, TDefaultErrorResult<FLobbiesModifyLobbyDataImpl>&& Result) mutable
 	{
 		const FCreateLobby::Params& Params = InAsyncOp.GetParams();
-		const TSharedRef<FLobbyDataEOS>& LobbyData = GetOpDataChecked<TSharedRef<FLobbyDataEOS>>(InAsyncOp, LobbyDataKeyName);
+		const TSharedRef<FLobbyDataEOS>& LobbyData = GetOpDataChecked<TSharedRef<FLobbyDataEOS>>(InAsyncOp, UE_ONLINE_LOBBY_EOS_KEY_NAME_LOBBY_DATA);
 
 		if (Result.IsError())
 		{
@@ -257,7 +257,7 @@ TOnlineAsyncOpHandle<FCreateLobby> FLobbiesEOSGS::CreateLobby(FCreateLobby::Para
 	.Then([this](TOnlineAsyncOp<FCreateLobby>& InAsyncOp)
 	{
 		const FCreateLobby::Params& Params = InAsyncOp.GetParams();
-		const TSharedRef<FLobbyDataEOS>& LobbyData = GetOpDataChecked<TSharedRef<FLobbyDataEOS>>(InAsyncOp, LobbyDataKeyName);
+		const TSharedRef<FLobbyDataEOS>& LobbyData = GetOpDataChecked<TSharedRef<FLobbyDataEOS>>(InAsyncOp, UE_ONLINE_LOBBY_EOS_KEY_NAME_LOBBY_DATA);
 
 		// Mark the lobby active.
 		AddActiveLobby(Params.LocalAccountId, LobbyData);
@@ -307,13 +307,13 @@ TOnlineAsyncOpHandle<FFindLobbies> FLobbiesEOSGS::FindLobbies(FFindLobbies::Para
 		}
 		else
 		{
-			InAsyncOp.Data.Set<TSharedRef<FLobbySearchEOS>>(LobbySearchKeyName, Result.GetOkValue());
+			InAsyncOp.Data.Set<TSharedRef<FLobbySearchEOS>>(UE_ONLINE_LOBBY_EOS_KEY_NAME_LOBBY_SEARCH, Result.GetOkValue());
 		}
 	})
 	.Then([this](TOnlineAsyncOp<FFindLobbies>& InAsyncOp)
 	{
 		const FFindLobbies::Params& Params = InAsyncOp.GetParams();
-		const TSharedRef<FLobbySearchEOS>& LobbySearch = GetOpDataChecked<TSharedRef<FLobbySearchEOS>>(InAsyncOp, LobbySearchKeyName);
+		const TSharedRef<FLobbySearchEOS>& LobbySearch = GetOpDataChecked<TSharedRef<FLobbySearchEOS>>(InAsyncOp, UE_ONLINE_LOBBY_EOS_KEY_NAME_LOBBY_SEARCH);
 
 		ActiveSearchResults.Add(InAsyncOp.GetParams().LocalAccountId, LobbySearch);
 
@@ -354,7 +354,7 @@ TOnlineAsyncOpHandle<FJoinLobby> FLobbiesEOSGS::JoinLobby(FJoinLobby::Params&& I
 		Op->SetError(Errors::InvalidParams());
 		return Op->GetHandle();
 	}
-	Op->Data.Set<TSharedRef<FLobbyDataEOS>>(LobbyDataKeyName, LobbyData.ToSharedRef());
+	Op->Data.Set<TSharedRef<FLobbyDataEOS>>(UE_ONLINE_LOBBY_EOS_KEY_NAME_LOBBY_DATA, LobbyData.ToSharedRef());
 
 	// Setup lobby details - Prefer UI event before invitation before search result.
 	TSharedPtr<FLobbyDetailsEOS> LobbyDetails = LobbyData->GetUserLobbyDetails(Params.LocalAccountId);
@@ -365,7 +365,7 @@ TOnlineAsyncOpHandle<FJoinLobby> FLobbiesEOSGS::JoinLobby(FJoinLobby::Params&& I
 		Op->SetError(Errors::InvalidParams());
 		return Op->GetHandle();
 	}
-	Op->Data.Set<TSharedRef<FLobbyDetailsEOS>>(LobbyDetailsKeyName, LobbyDetails.ToSharedRef());
+	Op->Data.Set<TSharedRef<FLobbyDetailsEOS>>(UE_ONLINE_LOBBY_EOS_KEY_NAME_LOBBY_DETAILS, LobbyDetails.ToSharedRef());
 
 	// Check that lobby is compatible with the running game.
 	if (!LobbyDetails->IsBucketCompatible())
@@ -381,7 +381,7 @@ TOnlineAsyncOpHandle<FJoinLobby> FLobbiesEOSGS::JoinLobby(FJoinLobby::Params&& I
 	Op->Then([this](TOnlineAsyncOp<FJoinLobby>& InAsyncOp, TPromise<const EOS_Lobby_JoinLobbyCallbackInfo*>&& Promise)
 	{
 		const FJoinLobby::Params& Params = InAsyncOp.GetParams();
-		const TSharedRef<FLobbyDetailsEOS>& LobbyDetails = GetOpDataChecked<TSharedRef<FLobbyDetailsEOS>>(InAsyncOp, LobbyDetailsKeyName);
+		const TSharedRef<FLobbyDetailsEOS>& LobbyDetails = GetOpDataChecked<TSharedRef<FLobbyDetailsEOS>>(InAsyncOp, UE_ONLINE_LOBBY_EOS_KEY_NAME_LOBBY_DETAILS);
 
 		EOS_Lobby_JoinLobbyOptions JoinLobbyOptions = {};
 		JoinLobbyOptions.ApiVersion = 3;
@@ -399,7 +399,7 @@ TOnlineAsyncOpHandle<FJoinLobby> FLobbiesEOSGS::JoinLobby(FJoinLobby::Params&& I
 		if (Data->ResultCode != EOS_EResult::EOS_Success)
 		{
 			const FJoinLobby::Params& Params = InAsyncOp.GetParams();
-			const TSharedRef<FLobbyDataEOS>& LobbyData = GetOpDataChecked<TSharedRef<FLobbyDataEOS>>(InAsyncOp, LobbyDataKeyName);
+			const TSharedRef<FLobbyDataEOS>& LobbyData = GetOpDataChecked<TSharedRef<FLobbyDataEOS>>(InAsyncOp, UE_ONLINE_LOBBY_EOS_KEY_NAME_LOBBY_DATA);
 
 			UE_LOG(LogTemp, Warning, TEXT("[FLobbiesEOSGS::JoinLobby] EOS_Lobby_JoinLobby Failed: User[%s], Lobby[%s], Result[%s]"),
 				*ToLogString(Params.LocalAccountId), *LobbyData->GetLobbyIdString(), *LexToString(Data->ResultCode));
@@ -411,7 +411,7 @@ TOnlineAsyncOpHandle<FJoinLobby> FLobbiesEOSGS::JoinLobby(FJoinLobby::Params&& I
 	.Then([this](TOnlineAsyncOp<FJoinLobby>& InAsyncOp)
 	{
 		const FJoinLobby::Params& Params = InAsyncOp.GetParams();
-		const TSharedRef<FLobbyDataEOS>& LobbyData = GetOpDataChecked<TSharedRef<FLobbyDataEOS>>(InAsyncOp, LobbyDataKeyName);
+		const TSharedRef<FLobbyDataEOS>& LobbyData = GetOpDataChecked<TSharedRef<FLobbyDataEOS>>(InAsyncOp, UE_ONLINE_LOBBY_EOS_KEY_NAME_LOBBY_DATA);
 
 		FLobbyClientDataPrepareClientChanges::Params PrepareParams;
 		PrepareParams.LocalAccountId = Params.LocalAccountId;
@@ -438,7 +438,7 @@ TOnlineAsyncOpHandle<FJoinLobby> FLobbiesEOSGS::JoinLobby(FJoinLobby::Params&& I
 		if (Result.IsError())
 		{
 			const FJoinLobby::Params& Params = InAsyncOp.GetParams();
-			const TSharedRef<FLobbyDataEOS>& LobbyData = GetOpDataChecked<TSharedRef<FLobbyDataEOS>>(InAsyncOp, LobbyDataKeyName);
+			const TSharedRef<FLobbyDataEOS>& LobbyData = GetOpDataChecked<TSharedRef<FLobbyDataEOS>>(InAsyncOp, UE_ONLINE_LOBBY_EOS_KEY_NAME_LOBBY_DATA);
 
 			UE_LOG(LogTemp, Warning, TEXT("[FLobbiesEOSGS::JoinLobby] ModifyLobbyMemberDataImpl Failed: User[%s], Lobby[%s], Result[%s]"),
 				*ToLogString(Params.LocalAccountId), *LobbyData->GetLobbyIdString(), *Result.GetErrorValue().GetLogString());
@@ -482,7 +482,7 @@ TOnlineAsyncOpHandle<FJoinLobby> FLobbiesEOSGS::JoinLobby(FJoinLobby::Params&& I
 	.Then([this](TOnlineAsyncOp<FJoinLobby>& InAsyncOp)
 	{
 		const FJoinLobby::Params& Params = InAsyncOp.GetParams();
-		const TSharedRef<FLobbyDataEOS>& LobbyData = GetOpDataChecked<TSharedRef<FLobbyDataEOS>>(InAsyncOp, LobbyDataKeyName);
+		const TSharedRef<FLobbyDataEOS>& LobbyData = GetOpDataChecked<TSharedRef<FLobbyDataEOS>>(InAsyncOp, UE_ONLINE_LOBBY_EOS_KEY_NAME_LOBBY_DATA);
 
 		// Mark the lobby active.
 		AddActiveLobby(Params.LocalAccountId, LobbyData);
@@ -526,13 +526,13 @@ TOnlineAsyncOpHandle<FLeaveLobby> FLobbiesEOSGS::LeaveLobby(FLeaveLobby::Params&
 		Op->SetError(Errors::InvalidParams());
 		return Op->GetHandle();
 	}
-	Op->Data.Set<TSharedRef<FLobbyDataEOS>>(LobbyDataKeyName, LobbyData.ToSharedRef());
+	Op->Data.Set<TSharedRef<FLobbyDataEOS>>(UE_ONLINE_LOBBY_EOS_KEY_NAME_LOBBY_DATA, LobbyData.ToSharedRef());
 
 	// Start operation.
 	Op->Then([this](TOnlineAsyncOp<FLeaveLobby>& InAsyncOp)
 	{
 		const FLeaveLobby::Params& Params = InAsyncOp.GetParams();
-		const TSharedRef<FLobbyDataEOS>& LobbyData = GetOpDataChecked<TSharedRef<FLobbyDataEOS>>(InAsyncOp, LobbyDataKeyName);
+		const TSharedRef<FLobbyDataEOS>& LobbyData = GetOpDataChecked<TSharedRef<FLobbyDataEOS>>(InAsyncOp, UE_ONLINE_LOBBY_EOS_KEY_NAME_LOBBY_DATA);
 
 		FLobbyClientDataPrepareClientChanges::Params PrepareParams;
 		PrepareParams.LocalAccountId = Params.LocalAccountId;
@@ -556,7 +556,7 @@ TOnlineAsyncOpHandle<FLeaveLobby> FLobbiesEOSGS::LeaveLobby(FLeaveLobby::Params&
 		if (Result.IsError())
 		{
 			const FLeaveLobby::Params& Params = InAsyncOp.GetParams();
-			const TSharedRef<FLobbyDataEOS>& LobbyData = GetOpDataChecked<TSharedRef<FLobbyDataEOS>>(InAsyncOp, LobbyDataKeyName);
+			const TSharedRef<FLobbyDataEOS>& LobbyData = GetOpDataChecked<TSharedRef<FLobbyDataEOS>>(InAsyncOp, UE_ONLINE_LOBBY_EOS_KEY_NAME_LOBBY_DATA);
 
 			UE_LOG(LogTemp, Warning, TEXT("[FLobbiesEOSGS::LeaveLobby] LeaveLobbyImpl Failed: User[%s], Lobby[%s], Result[%s]"),
 				*ToLogString(Params.LocalAccountId), *LobbyData->GetLobbyIdString(), *Result.GetErrorValue().GetLogString());
@@ -570,7 +570,7 @@ TOnlineAsyncOpHandle<FLeaveLobby> FLobbiesEOSGS::LeaveLobby(FLeaveLobby::Params&
 	.Then([this](TOnlineAsyncOp<FLeaveLobby>& InAsyncOp)
 	{
 		const FLeaveLobby::Params& Params = InAsyncOp.GetParams();
-		const TSharedRef<FLobbyDataEOS>& LobbyData = GetOpDataChecked<TSharedRef<FLobbyDataEOS>>(InAsyncOp, LobbyDataKeyName);
+		const TSharedRef<FLobbyDataEOS>& LobbyData = GetOpDataChecked<TSharedRef<FLobbyDataEOS>>(InAsyncOp, UE_ONLINE_LOBBY_EOS_KEY_NAME_LOBBY_DATA);
 
 		// Update local cache and fire events.
 		LobbyData->GetLobbyClientData()->CommitClientChanges({ &LobbyEvents });
@@ -615,13 +615,13 @@ TOnlineAsyncOpHandle<FInviteLobbyMember> FLobbiesEOSGS::InviteLobbyMember(FInvit
 		Op->SetError(Errors::InvalidParams());
 		return Op->GetHandle();
 	}
-	Op->Data.Set<TSharedRef<FLobbyDataEOS>>(LobbyDataKeyName, LobbyData.ToSharedRef());
+	Op->Data.Set<TSharedRef<FLobbyDataEOS>>(UE_ONLINE_LOBBY_EOS_KEY_NAME_LOBBY_DATA, LobbyData.ToSharedRef());
 
 	// Start operation.
 	Op->Then([this](TOnlineAsyncOp<FInviteLobbyMember>& InAsyncOp)
 	{
 		const FInviteLobbyMember::Params& Params = InAsyncOp.GetParams();
-		const TSharedRef<FLobbyDataEOS>& LobbyData = GetOpDataChecked<TSharedRef<FLobbyDataEOS>>(InAsyncOp, LobbyDataKeyName);
+		const TSharedRef<FLobbyDataEOS>& LobbyData = GetOpDataChecked<TSharedRef<FLobbyDataEOS>>(InAsyncOp, UE_ONLINE_LOBBY_EOS_KEY_NAME_LOBBY_DATA);
 
 		FLobbiesInviteLobbyMemberImpl::Params InviteParams;
 		InviteParams.LobbyData = LobbyData;
@@ -634,7 +634,7 @@ TOnlineAsyncOpHandle<FInviteLobbyMember> FLobbiesEOSGS::InviteLobbyMember(FInvit
 		if (Result.IsError())
 		{
 			const FInviteLobbyMember::Params& Params = InAsyncOp.GetParams();
-			const TSharedRef<FLobbyDataEOS>& LobbyData = GetOpDataChecked<TSharedRef<FLobbyDataEOS>>(InAsyncOp, LobbyDataKeyName);
+			const TSharedRef<FLobbyDataEOS>& LobbyData = GetOpDataChecked<TSharedRef<FLobbyDataEOS>>(InAsyncOp, UE_ONLINE_LOBBY_EOS_KEY_NAME_LOBBY_DATA);
 
 			UE_LOG(LogTemp, Warning, TEXT("[FLobbiesEOSGS::InviteLobbyMember] InviteLobbyMemberImpl Failed: User[%s], Lobby[%s], Result[%s]"),
 				*ToLogString(Params.LocalAccountId), *LobbyData->GetLobbyIdString(), *Result.GetErrorValue().GetLogString());
@@ -645,7 +645,7 @@ TOnlineAsyncOpHandle<FInviteLobbyMember> FLobbiesEOSGS::InviteLobbyMember(FInvit
 	.Then([this](TOnlineAsyncOp<FInviteLobbyMember>& InAsyncOp)
 	{
 		const FInviteLobbyMember::Params& Params = InAsyncOp.GetParams();
-		const TSharedRef<FLobbyDataEOS>& LobbyData = GetOpDataChecked<TSharedRef<FLobbyDataEOS>>(InAsyncOp, LobbyDataKeyName);
+		const TSharedRef<FLobbyDataEOS>& LobbyData = GetOpDataChecked<TSharedRef<FLobbyDataEOS>>(InAsyncOp, UE_ONLINE_LOBBY_EOS_KEY_NAME_LOBBY_DATA);
 
 		UE_LOG(LogTemp, Verbose, TEXT("[FLobbiesEOSGS::InviteLobbyMember] Succeeded: User[%s], Lobby[%s]"),
 			*ToLogString(Params.LocalAccountId), *LobbyData->GetLobbyIdString());
@@ -691,7 +691,7 @@ TOnlineAsyncOpHandle<FDeclineLobbyInvitation> FLobbiesEOSGS::DeclineLobbyInvitat
 		if (Result.IsError())
 		{
 			const FDeclineLobbyInvitation::Params& Params = InAsyncOp.GetParams();
-			const TSharedRef<FLobbyDataEOS>& LobbyData = GetOpDataChecked<TSharedRef<FLobbyDataEOS>>(InAsyncOp, LobbyDataKeyName);
+			const TSharedRef<FLobbyDataEOS>& LobbyData = GetOpDataChecked<TSharedRef<FLobbyDataEOS>>(InAsyncOp, UE_ONLINE_LOBBY_EOS_KEY_NAME_LOBBY_DATA);
 
 			UE_LOG(LogTemp, Warning, TEXT("[FLobbiesEOSGS::DeclineLobbyInvitation] DeclineLobbyInvitationImpl Failed: User[%s], Lobby[%s], Result[%s]"),
 				*ToLogString(Params.LocalAccountId), *LobbyData->GetLobbyIdString(), *Result.GetErrorValue().GetLogString());
@@ -702,7 +702,7 @@ TOnlineAsyncOpHandle<FDeclineLobbyInvitation> FLobbiesEOSGS::DeclineLobbyInvitat
 	.Then([this](TOnlineAsyncOp<FDeclineLobbyInvitation>& InAsyncOp)
 	{
 		const FDeclineLobbyInvitation::Params& Params = InAsyncOp.GetParams();
-		const TSharedRef<FLobbyDataEOS>& LobbyData = GetOpDataChecked<TSharedRef<FLobbyDataEOS>>(InAsyncOp, LobbyDataKeyName);
+		const TSharedRef<FLobbyDataEOS>& LobbyData = GetOpDataChecked<TSharedRef<FLobbyDataEOS>>(InAsyncOp, UE_ONLINE_LOBBY_EOS_KEY_NAME_LOBBY_DATA);
 
 		UE_LOG(LogTemp, Verbose, TEXT("[FLobbiesEOSGS::DeclineLobbyInvitation] Succeeded: User[%s], Lobby[%s]"),
 			*ToLogString(Params.LocalAccountId), *LobbyData->GetLobbyIdString());
@@ -740,13 +740,13 @@ TOnlineAsyncOpHandle<FKickLobbyMember> FLobbiesEOSGS::KickLobbyMember(FKickLobby
 		Op->SetError(Errors::InvalidParams());
 		return Op->GetHandle();
 	}
-	Op->Data.Set<TSharedRef<FLobbyDataEOS>>(LobbyDataKeyName, LobbyData.ToSharedRef());
+	Op->Data.Set<TSharedRef<FLobbyDataEOS>>(UE_ONLINE_LOBBY_EOS_KEY_NAME_LOBBY_DATA, LobbyData.ToSharedRef());
 
 	// Start operation.
 	Op->Then([this](TOnlineAsyncOp<FKickLobbyMember>& InAsyncOp)
 	{
 		const FKickLobbyMember::Params& Params = InAsyncOp.GetParams();
-		const TSharedRef<FLobbyDataEOS>& LobbyData = GetOpDataChecked<TSharedRef<FLobbyDataEOS>>(InAsyncOp, LobbyDataKeyName);
+		const TSharedRef<FLobbyDataEOS>& LobbyData = GetOpDataChecked<TSharedRef<FLobbyDataEOS>>(InAsyncOp, UE_ONLINE_LOBBY_EOS_KEY_NAME_LOBBY_DATA);
 
 		FLobbyClientDataPrepareClientChanges::Params PrepareParams;
 		PrepareParams.LocalAccountId = Params.LocalAccountId;
@@ -771,7 +771,7 @@ TOnlineAsyncOpHandle<FKickLobbyMember> FLobbiesEOSGS::KickLobbyMember(FKickLobby
 		if (Result.IsError())
 		{
 			const FKickLobbyMember::Params& Params = InAsyncOp.GetParams();
-			const TSharedRef<FLobbyDataEOS>& LobbyData = GetOpDataChecked<TSharedRef<FLobbyDataEOS>>(InAsyncOp, LobbyDataKeyName);
+			const TSharedRef<FLobbyDataEOS>& LobbyData = GetOpDataChecked<TSharedRef<FLobbyDataEOS>>(InAsyncOp, UE_ONLINE_LOBBY_EOS_KEY_NAME_LOBBY_DATA);
 
 			UE_LOG(LogTemp, Warning, TEXT("[FLobbiesEOSGS::KickLobbyMember] KickLobbyMemberImpl Failed: User[%s], Lobby[%s], Result[%s]"),
 				*ToLogString(Params.LocalAccountId), *LobbyData->GetLobbyIdString(), *Result.GetErrorValue().GetLogString());
@@ -781,7 +781,7 @@ TOnlineAsyncOpHandle<FKickLobbyMember> FLobbiesEOSGS::KickLobbyMember(FKickLobby
 	.Then([this](TOnlineAsyncOp<FKickLobbyMember>& InAsyncOp)
 	{
 		const FKickLobbyMember::Params& Params = InAsyncOp.GetParams();
-		const TSharedRef<FLobbyDataEOS>& LobbyData = GetOpDataChecked<TSharedRef<FLobbyDataEOS>>(InAsyncOp, LobbyDataKeyName);
+		const TSharedRef<FLobbyDataEOS>& LobbyData = GetOpDataChecked<TSharedRef<FLobbyDataEOS>>(InAsyncOp, UE_ONLINE_LOBBY_EOS_KEY_NAME_LOBBY_DATA);
 
 		// Update local cache and fire events.
 		LobbyData->GetLobbyClientData()->CommitClientChanges({ &LobbyEvents });
@@ -822,13 +822,13 @@ TOnlineAsyncOpHandle<FPromoteLobbyMember> FLobbiesEOSGS::PromoteLobbyMember(FPro
 		Op->SetError(Errors::InvalidParams());
 		return Op->GetHandle();
 	}
-	Op->Data.Set<TSharedRef<FLobbyDataEOS>>(LobbyDataKeyName, LobbyData.ToSharedRef());
+	Op->Data.Set<TSharedRef<FLobbyDataEOS>>(UE_ONLINE_LOBBY_EOS_KEY_NAME_LOBBY_DATA, LobbyData.ToSharedRef());
 
 	// Start operation.
 	Op->Then([this](TOnlineAsyncOp<FPromoteLobbyMember>& InAsyncOp)
 	{
 		const FPromoteLobbyMember::Params& Params = InAsyncOp.GetParams();
-		const TSharedRef<FLobbyDataEOS>& LobbyData = GetOpDataChecked<TSharedRef<FLobbyDataEOS>>(InAsyncOp, LobbyDataKeyName);
+		const TSharedRef<FLobbyDataEOS>& LobbyData = GetOpDataChecked<TSharedRef<FLobbyDataEOS>>(InAsyncOp, UE_ONLINE_LOBBY_EOS_KEY_NAME_LOBBY_DATA);
 
 		FLobbyClientDataPrepareClientChanges::Params PrepareParams;
 		PrepareParams.LocalAccountId = Params.LocalAccountId;
@@ -853,7 +853,7 @@ TOnlineAsyncOpHandle<FPromoteLobbyMember> FLobbiesEOSGS::PromoteLobbyMember(FPro
 		if (Result.IsError())
 		{
 			const FPromoteLobbyMember::Params& Params = InAsyncOp.GetParams();
-			const TSharedRef<FLobbyDataEOS>& LobbyData = GetOpDataChecked<TSharedRef<FLobbyDataEOS>>(InAsyncOp, LobbyDataKeyName);
+			const TSharedRef<FLobbyDataEOS>& LobbyData = GetOpDataChecked<TSharedRef<FLobbyDataEOS>>(InAsyncOp, UE_ONLINE_LOBBY_EOS_KEY_NAME_LOBBY_DATA);
 
 			UE_LOG(LogTemp, Warning, TEXT("[FLobbiesEOSGS::FPromoteLobbyMember] PromoteLobbyMemberImpl Failed: User[%s], Lobby[%s], Result[%s]"),
 				*ToLogString(Params.LocalAccountId), *LobbyData->GetLobbyIdString(), *Result.GetErrorValue().GetLogString());
@@ -864,7 +864,7 @@ TOnlineAsyncOpHandle<FPromoteLobbyMember> FLobbiesEOSGS::PromoteLobbyMember(FPro
 	.Then([this](TOnlineAsyncOp<FPromoteLobbyMember>& InAsyncOp)
 	{
 		const FPromoteLobbyMember::Params& Params = InAsyncOp.GetParams();
-		const TSharedRef<FLobbyDataEOS>& LobbyData = GetOpDataChecked<TSharedRef<FLobbyDataEOS>>(InAsyncOp, LobbyDataKeyName);
+		const TSharedRef<FLobbyDataEOS>& LobbyData = GetOpDataChecked<TSharedRef<FLobbyDataEOS>>(InAsyncOp, UE_ONLINE_LOBBY_EOS_KEY_NAME_LOBBY_DATA);
 
 		// Update local cache and fire events.
 		LobbyData->GetLobbyClientData()->CommitClientChanges({ &LobbyEvents });
@@ -905,13 +905,13 @@ TOnlineAsyncOpHandle<FModifyLobbyJoinPolicy> FLobbiesEOSGS::ModifyLobbyJoinPolic
 		Op->SetError(Errors::InvalidParams());
 		return Op->GetHandle();
 	}
-	Op->Data.Set<TSharedRef<FLobbyDataEOS>>(LobbyDataKeyName, LobbyData.ToSharedRef());
+	Op->Data.Set<TSharedRef<FLobbyDataEOS>>(UE_ONLINE_LOBBY_EOS_KEY_NAME_LOBBY_DATA, LobbyData.ToSharedRef());
 
 	// Start operation.
 	Op->Then([this](TOnlineAsyncOp<FModifyLobbyJoinPolicy>& InAsyncOp)
 	{
 		const FModifyLobbyJoinPolicy::Params& Params = InAsyncOp.GetParams();
-		const TSharedRef<FLobbyDataEOS>& LobbyData = GetOpDataChecked<TSharedRef<FLobbyDataEOS>>(InAsyncOp, LobbyDataKeyName);
+		const TSharedRef<FLobbyDataEOS>& LobbyData = GetOpDataChecked<TSharedRef<FLobbyDataEOS>>(InAsyncOp, UE_ONLINE_LOBBY_EOS_KEY_NAME_LOBBY_DATA);
 
 		FLobbyClientDataPrepareClientChanges::Params PrepareParams;
 		PrepareParams.LocalAccountId = Params.LocalAccountId;
@@ -936,7 +936,7 @@ TOnlineAsyncOpHandle<FModifyLobbyJoinPolicy> FLobbiesEOSGS::ModifyLobbyJoinPolic
 		if (Result.IsError())
 		{
 			const FModifyLobbyJoinPolicy::Params& Params = InAsyncOp.GetParams();
-			const TSharedRef<FLobbyDataEOS>& LobbyData = GetOpDataChecked<TSharedRef<FLobbyDataEOS>>(InAsyncOp, LobbyDataKeyName);
+			const TSharedRef<FLobbyDataEOS>& LobbyData = GetOpDataChecked<TSharedRef<FLobbyDataEOS>>(InAsyncOp, UE_ONLINE_LOBBY_EOS_KEY_NAME_LOBBY_DATA);
 
 			UE_LOG(LogTemp, Warning, TEXT("[FLobbiesEOSGS::ModifyLobbyJoinPolicy] ModifyLobbyDataImpl Failed: User[%s], Lobby[%s], Result[%s]"),
 				*ToLogString(Params.LocalAccountId), *LobbyData->GetLobbyIdString(), *Result.GetErrorValue().GetLogString());
@@ -947,7 +947,7 @@ TOnlineAsyncOpHandle<FModifyLobbyJoinPolicy> FLobbiesEOSGS::ModifyLobbyJoinPolic
 	.Then([this](TOnlineAsyncOp<FModifyLobbyJoinPolicy>& InAsyncOp)
 	{
 		const FModifyLobbyJoinPolicy::Params& Params = InAsyncOp.GetParams();
-		const TSharedRef<FLobbyDataEOS>& LobbyData = GetOpDataChecked<TSharedRef<FLobbyDataEOS>>(InAsyncOp, LobbyDataKeyName);
+		const TSharedRef<FLobbyDataEOS>& LobbyData = GetOpDataChecked<TSharedRef<FLobbyDataEOS>>(InAsyncOp, UE_ONLINE_LOBBY_EOS_KEY_NAME_LOBBY_DATA);
 
 		// Update local cache and fire events.
 		LobbyData->GetLobbyClientData()->CommitClientChanges({ &LobbyEvents });
@@ -988,13 +988,13 @@ TOnlineAsyncOpHandle<FModifyLobbyAttributes> FLobbiesEOSGS::ModifyLobbyAttribute
 		Op->SetError(Errors::InvalidParams());
 		return Op->GetHandle();
 	}
-	Op->Data.Set<TSharedRef<FLobbyDataEOS>>(LobbyDataKeyName, LobbyData.ToSharedRef());
+	Op->Data.Set<TSharedRef<FLobbyDataEOS>>(UE_ONLINE_LOBBY_EOS_KEY_NAME_LOBBY_DATA, LobbyData.ToSharedRef());
 
 	// Start operation.
 	Op->Then([this](TOnlineAsyncOp<FModifyLobbyAttributes>& InAsyncOp)
 	{
 		const FModifyLobbyAttributes::Params& Params = InAsyncOp.GetParams();
-		const TSharedRef<FLobbyDataEOS>& LobbyData = GetOpDataChecked<TSharedRef<FLobbyDataEOS>>(InAsyncOp, LobbyDataKeyName);
+		const TSharedRef<FLobbyDataEOS>& LobbyData = GetOpDataChecked<TSharedRef<FLobbyDataEOS>>(InAsyncOp, UE_ONLINE_LOBBY_EOS_KEY_NAME_LOBBY_DATA);
 
 		FLobbyClientDataPrepareClientChanges::Params PrepareParams;
 		PrepareParams.LocalAccountId = Params.LocalAccountId;
@@ -1019,7 +1019,7 @@ TOnlineAsyncOpHandle<FModifyLobbyAttributes> FLobbiesEOSGS::ModifyLobbyAttribute
 		if (Result.IsError())
 		{
 			const FModifyLobbyAttributes::Params& Params = InAsyncOp.GetParams();
-			const TSharedRef<FLobbyDataEOS>& LobbyData = GetOpDataChecked<TSharedRef<FLobbyDataEOS>>(InAsyncOp, LobbyDataKeyName);
+			const TSharedRef<FLobbyDataEOS>& LobbyData = GetOpDataChecked<TSharedRef<FLobbyDataEOS>>(InAsyncOp, UE_ONLINE_LOBBY_EOS_KEY_NAME_LOBBY_DATA);
 
 			UE_LOG(LogTemp, Warning, TEXT("[FLobbiesEOSGS::ModifyLobbyAttributes] ModifyLobbyDataImpl Failed: User[%s], Lobby[%s], Result[%s]"),
 				*ToLogString(Params.LocalAccountId), *LobbyData->GetLobbyIdString(), *Result.GetErrorValue().GetLogString());
@@ -1030,7 +1030,7 @@ TOnlineAsyncOpHandle<FModifyLobbyAttributes> FLobbiesEOSGS::ModifyLobbyAttribute
 	.Then([this](TOnlineAsyncOp<FModifyLobbyAttributes>& InAsyncOp)
 	{
 		const FModifyLobbyAttributes::Params& Params = InAsyncOp.GetParams();
-		const TSharedRef<FLobbyDataEOS>& LobbyData = GetOpDataChecked<TSharedRef<FLobbyDataEOS>>(InAsyncOp, LobbyDataKeyName);
+		const TSharedRef<FLobbyDataEOS>& LobbyData = GetOpDataChecked<TSharedRef<FLobbyDataEOS>>(InAsyncOp, UE_ONLINE_LOBBY_EOS_KEY_NAME_LOBBY_DATA);
 
 		// Update local cache and fire events.
 		LobbyData->GetLobbyClientData()->CommitClientChanges({ &LobbyEvents });
@@ -1071,13 +1071,13 @@ TOnlineAsyncOpHandle<FModifyLobbyMemberAttributes> FLobbiesEOSGS::ModifyLobbyMem
 		Op->SetError(Errors::InvalidParams());
 		return Op->GetHandle();
 	}
-	Op->Data.Set<TSharedRef<FLobbyDataEOS>>(LobbyDataKeyName, LobbyData.ToSharedRef());
+	Op->Data.Set<TSharedRef<FLobbyDataEOS>>(UE_ONLINE_LOBBY_EOS_KEY_NAME_LOBBY_DATA, LobbyData.ToSharedRef());
 
 	// Start operation.
 	Op->Then([this](TOnlineAsyncOp<FModifyLobbyMemberAttributes>& InAsyncOp)
 	{
 		const FModifyLobbyMemberAttributes::Params& Params = InAsyncOp.GetParams();
-		const TSharedRef<FLobbyDataEOS>& LobbyData = GetOpDataChecked<TSharedRef<FLobbyDataEOS>>(InAsyncOp, LobbyDataKeyName);
+		const TSharedRef<FLobbyDataEOS>& LobbyData = GetOpDataChecked<TSharedRef<FLobbyDataEOS>>(InAsyncOp, UE_ONLINE_LOBBY_EOS_KEY_NAME_LOBBY_DATA);
 
 		FLobbyClientDataPrepareClientChanges::Params PrepareParams;
 		PrepareParams.LocalAccountId = Params.LocalAccountId;
@@ -1102,7 +1102,7 @@ TOnlineAsyncOpHandle<FModifyLobbyMemberAttributes> FLobbiesEOSGS::ModifyLobbyMem
 		if (Result.IsError())
 		{
 			const FModifyLobbyMemberAttributes::Params& Params = InAsyncOp.GetParams();
-			const TSharedRef<FLobbyDataEOS>& LobbyData = GetOpDataChecked<TSharedRef<FLobbyDataEOS>>(InAsyncOp, LobbyDataKeyName);
+			const TSharedRef<FLobbyDataEOS>& LobbyData = GetOpDataChecked<TSharedRef<FLobbyDataEOS>>(InAsyncOp, UE_ONLINE_LOBBY_EOS_KEY_NAME_LOBBY_DATA);
 
 			UE_LOG(LogTemp, Warning, TEXT("[FLobbiesEOSGS::ModifyLobbyMemberAttributes] ModifyLobbyMemberDataImpl Failed: User[%s], Lobby[%s], Result[%s]"),
 				*ToLogString(Params.LocalAccountId), *LobbyData->GetLobbyIdString(), *Result.GetErrorValue().GetLogString());
@@ -1113,7 +1113,7 @@ TOnlineAsyncOpHandle<FModifyLobbyMemberAttributes> FLobbiesEOSGS::ModifyLobbyMem
 	.Then([this](TOnlineAsyncOp<FModifyLobbyMemberAttributes>& InAsyncOp)
 	{
 		const FModifyLobbyMemberAttributes::Params& Params = InAsyncOp.GetParams();
-		const TSharedRef<FLobbyDataEOS>& LobbyData = GetOpDataChecked<TSharedRef<FLobbyDataEOS>>(InAsyncOp, LobbyDataKeyName);
+		const TSharedRef<FLobbyDataEOS>& LobbyData = GetOpDataChecked<TSharedRef<FLobbyDataEOS>>(InAsyncOp, UE_ONLINE_LOBBY_EOS_KEY_NAME_LOBBY_DATA);
 
 		// Update local cache and fire events.
 		LobbyData->GetLobbyClientData()->CommitClientChanges({ &LobbyEvents });
@@ -1921,7 +1921,7 @@ TOnlineAsyncOpHandle<FLobbiesProcessLobbyNotificationImpl> FLobbiesEOSGS::Proces
 			return MakeFulfilledPromise<TDefaultErrorResultInternal<FLobbyServiceSnapshot>>(Errors::InvalidState()).GetFuture();
 		}
 
-		InAsyncOp.Data.Set<TSharedRef<FLobbyDetailsEOS>>(LobbyDetailsKeyName, LobbyDetails.ToSharedRef());
+		InAsyncOp.Data.Set<TSharedRef<FLobbyDetailsEOS>>(UE_ONLINE_LOBBY_EOS_KEY_NAME_LOBBY_DETAILS, LobbyDetails.ToSharedRef());
 
 		// Fetch lobby snapshot. Fetching the snapshot resolves the account ids of all lobby members in the snapshot.
 		return LobbyDetails->GetLobbySnapshot();
@@ -1929,7 +1929,7 @@ TOnlineAsyncOpHandle<FLobbiesProcessLobbyNotificationImpl> FLobbiesEOSGS::Proces
 	.Then([this](TOnlineAsyncOp<FLobbiesProcessLobbyNotificationImpl>& InAsyncOp, TDefaultErrorResultInternal<FLobbyServiceSnapshot>&& LobbySnapshotResult)
 	{
 		const FLobbiesProcessLobbyNotificationImpl::Params& Params = InAsyncOp.GetParams();
-		TSharedRef<FLobbyDetailsEOS> LobbyDetails = GetOpDataChecked<TSharedRef<FLobbyDetailsEOS>>(InAsyncOp, LobbyDetailsKeyName);
+		TSharedRef<FLobbyDetailsEOS> LobbyDetails = GetOpDataChecked<TSharedRef<FLobbyDetailsEOS>>(InAsyncOp, UE_ONLINE_LOBBY_EOS_KEY_NAME_LOBBY_DETAILS);
 
 		if (LobbySnapshotResult.IsError())
 		{
diff --git a/Engine/Plugins/Online/OnlineServicesEOSGS/Source/Private/Online/LobbiesEOSGSTypes.cpp b/Engine/Plugins/Online/OnlineServicesEOSGS/Source/Private/Online/LobbiesEOSGSTypes.cpp
--- a/Engine/Plugins/Online/OnlineServicesEOSGS/Source/Private/Online/LobbiesEOSGSTypes.cpp
+++ b/Engine/Plugins/Online/OnlineServicesEOSGS/Source/Private/Online/LobbiesEOSGSTypes.cpp
@@ -42,10 +42,12 @@ FString ToLogString(const FLobbyBucketIdEOS& BucketId)
 	return FString::Printf(TEXT("%s:%d"), *BucketId.GetProductName(), BucketId.GetProductVersion());
 }
 
+PRAGMA_DISABLE_DEPRECATION_WARNINGS
 const FString FLobbyBucketIdEOS::Separator = TEXT("|");
+PRAGMA_ENABLE_DEPRECATION_WARNINGS
 
 FLobbyBucketIdEOS::FLobbyBucketIdEOS(FString ProductName, int32 ProductVersion)
-	: ProductName(ProductName.Replace(*Separator, TEXT("_")))
+	: ProductName(ProductName.Replace(UE_ONLINE_LOBBIES_BUCKET_ID_EOS_SEPARATOR, TEXT("_")))
 	, ProductVersion(ProductVersion)
 {
 }
@@ -125,7 +127,7 @@ FLobbyAttributeTranslator<ELobbyTranslationType::FromService>::FLobbyAttributeTr
 }
 
 FLobbyBucketIdTranslator<ELobbyTranslationType::ToService>::FLobbyBucketIdTranslator(const FLobbyBucketIdEOS& BucketId)
-	: BucketConverterStorage(*FString::Printf(TEXT("%s%s%d"), *BucketId.GetProductName(), *FLobbyBucketIdEOS::Separator, BucketId.GetProductVersion()))
+	: BucketConverterStorage(*FString::Printf(TEXT("%s%s%d"), *BucketId.GetProductName(), UE_ONLINE_LOBBIES_BUCKET_ID_EOS_SEPARATOR, BucketId.GetProductVersion()))
 {
 }
 
@@ -136,7 +138,7 @@ FLobbyBucketIdTranslator<ELobbyTranslationType::FromService>::FLobbyBucketIdTran
 
 	constexpr int32 ExpectedPartsNum = 2;
 	TArray<FString> Parts;
-	if (BucketString.ParseIntoArray(Parts, *FLobbyBucketIdEOS::Separator) == ExpectedPartsNum)
+	if (BucketString.ParseIntoArray(Parts, UE_ONLINE_LOBBIES_BUCKET_ID_EOS_SEPARATOR) == ExpectedPartsNum)
 	{
 		int32 BuildId = 0;
 		::LexFromString(BuildId, *Parts[1]);
diff --git a/Engine/Plugins/Online/OnlineServicesEOSGS/Source/Private/Online/LobbiesEOSGSTypes.h b/Engine/Plugins/Online/OnlineServicesEOSGS/Source/Private/Online/LobbiesEOSGSTypes.h
--- a/Engine/Plugins/Online/OnlineServicesEOSGS/Source/Private/Online/LobbiesEOSGSTypes.h
+++ b/Engine/Plugins/Online/OnlineServicesEOSGS/Source/Private/Online/LobbiesEOSGSTypes.h
@@ -32,9 +32,12 @@ struct FLobbyMemberServiceSnapshot;
 // Translated types
 //--------------------------------------------------------------------------------------------------
 
+#define UE_ONLINE_LOBBIES_BUCKET_ID_EOS_SEPARATOR TEXT("|")
+
 class FLobbyBucketIdEOS
 {
 public:
+	UE_DEPRECATED(5.2, "Separator is deprecated, please use UE_ONLINE_LOBBIES_BUCKET_ID_EOS_SEPARATOR instead.")
 	static const FString Separator;
 
 	FLobbyBucketIdEOS() = default;
diff --git a/Engine/Plugins/Online/OnlineServicesOSSAdapter/Source/Private/Online/AuthOSSAdapter.cpp b/Engine/Plugins/Online/OnlineServicesOSSAdapter/Source/Private/Online/AuthOSSAdapter.cpp
--- a/Engine/Plugins/Online/OnlineServicesOSSAdapter/Source/Private/Online/AuthOSSAdapter.cpp
+++ b/Engine/Plugins/Online/OnlineServicesOSSAdapter/Source/Private/Online/AuthOSSAdapter.cpp
@@ -13,16 +13,15 @@
 #include "OnlineSubsystem.h"
 #include "Interfaces/OnlineIdentityInterface.h"
 
+#define UE_ONLINE_AUTH_OSS_ADAPTER_KEY_NAME_INITIAL_LOCAL_USER_NUM TEXT("InitialLocalUserNum")
+#define UE_ONLINE_AUTH_OSS_ADAPTER_KEY_NAME_ACCOUNT_INFO TEXT("AccountInfoOSSAdapter")
+#define UE_ONLINE_AUTH_OSS_ADAPTER_KEY_NAME_EXTERNAL_AUTH_TOKEN_TRANSLATION_TRAITS TEXT("ExternalAuthTokenTranslationTraits")
+#define UE_ONLINE_AUTH_OSS_ADAPTER_KEY_NAME_EXTERNAL_SERVER_AUTH_TICKET_TRANSLATION_TRAITS TEXT("ExternalServerAuthTicketTranslationTraits")
 
 namespace UE::Online {
 
 namespace {
 
-static const FString InitialLocalUserNumKeyName = TEXT("InitialLocalUserNum");
-static const FString AccountInfoKeyName = TEXT("AccountInfoOSSAdapter");
-static const FString ExternalAuthTokenTranslationTraitsKeyName = TEXT("ExternalAuthTokenTranslationTraits");
-static const FString ExternalServerAuthTicketTranslationTraitsKeyName = TEXT("ExternalServerAuthTicketTranslationTraits");
-
 ELoginStatus TranslateLoginStatus(::ELoginStatus::Type Status)
 {
 	switch (Status)
@@ -180,7 +179,7 @@ TOnlineAsyncOpHandle<FAuthLogin> FAuthOSSAdapter::Login(FAuthLogin::Params&& Par
 
 		// Set initial user num on operation. Depending on the implementation the local user num
 		// which completes login may be different from the one which started it.
-		InAsyncOp.Data.Set<int32>(InitialLocalUserNumKeyName, InitialLocalUserNum);
+		InAsyncOp.Data.Set<int32>(UE_ONLINE_AUTH_OSS_ADAPTER_KEY_NAME_INITIAL_LOCAL_USER_NUM, InitialLocalUserNum);
 
 		if (AccountInfoRegistryOSSAdapter.Find(Params.PlatformUserId))
 		{
@@ -194,13 +193,13 @@ TOnlineAsyncOpHandle<FAuthLogin> FAuthOSSAdapter::Login(FAuthLogin::Params&& Par
 		AccountInfoOSSAdapter->LoginStatus = ELoginStatus::NotLoggedIn;
 
 		// Set user auth data on operation.
-		InAsyncOp.Data.Set<TSharedRef<FAccountInfoOSSAdapter>>(AccountInfoKeyName, AccountInfoOSSAdapter.ToSharedRef());
+		InAsyncOp.Data.Set<TSharedRef<FAccountInfoOSSAdapter>>(UE_ONLINE_AUTH_OSS_ADAPTER_KEY_NAME_ACCOUNT_INFO, AccountInfoOSSAdapter.ToSharedRef());
 	})
 	// Step 2: Login to OSSv1 identity interface.
 	.Then([this](TOnlineAsyncOp<FAuthLogin>& InAsyncOp)
 	{
-		const TSharedRef<FAccountInfoOSSAdapter>& AccountInfoOSSAdapter = GetOpDataChecked<TSharedRef<FAccountInfoOSSAdapter>>(InAsyncOp, AccountInfoKeyName);
-		const int InitialLocalUserNum = GetOpDataChecked<int32>(InAsyncOp, InitialLocalUserNumKeyName);
+		const TSharedRef<FAccountInfoOSSAdapter>& AccountInfoOSSAdapter = GetOpDataChecked<TSharedRef<FAccountInfoOSSAdapter>>(InAsyncOp, UE_ONLINE_AUTH_OSS_ADAPTER_KEY_NAME_ACCOUNT_INFO);
+		const int InitialLocalUserNum = GetOpDataChecked<int32>(InAsyncOp, UE_ONLINE_AUTH_OSS_ADAPTER_KEY_NAME_INITIAL_LOCAL_USER_NUM);
 
 		TPromise<void> Promise;
 		TFuture<void> Future = Promise.GetFuture();
@@ -254,7 +253,7 @@ TOnlineAsyncOpHandle<FAuthLogin> FAuthOSSAdapter::Login(FAuthLogin::Params&& Par
 	// Step 4: Fetch dependent data.
 	.Then([this](TOnlineAsyncOp<FAuthLogin>& InAsyncOp)
 	{
-		const TSharedRef<FAccountInfoOSSAdapter>& AccountInfoOSSAdapter = GetOpDataChecked<TSharedRef<FAccountInfoOSSAdapter>>(InAsyncOp, AccountInfoKeyName);
+		const TSharedRef<FAccountInfoOSSAdapter>& AccountInfoOSSAdapter = GetOpDataChecked<TSharedRef<FAccountInfoOSSAdapter>>(InAsyncOp, UE_ONLINE_AUTH_OSS_ADAPTER_KEY_NAME_ACCOUNT_INFO);
 
 		if (TSharedPtr<FUserOnlineAccount> UserOnlineAccount = GetIdentityInterface()->GetUserAccount(*AccountInfoOSSAdapter->UniqueNetId))
 		{
@@ -286,7 +285,7 @@ TOnlineAsyncOpHandle<FAuthLogin> FAuthOSSAdapter::Login(FAuthLogin::Params&& Par
 	// Step 6: bookkeeping and notifications.
 	.Then([this](TOnlineAsyncOp<FAuthLogin>& InAsyncOp)
 	{
-		const TSharedRef<FAccountInfoOSSAdapter>& AccountInfoOSSAdapter = GetOpDataChecked<TSharedRef<FAccountInfoOSSAdapter>>(InAsyncOp, AccountInfoKeyName);
+		const TSharedRef<FAccountInfoOSSAdapter>& AccountInfoOSSAdapter = GetOpDataChecked<TSharedRef<FAccountInfoOSSAdapter>>(InAsyncOp, UE_ONLINE_AUTH_OSS_ADAPTER_KEY_NAME_ACCOUNT_INFO);
 
 		AccountInfoOSSAdapter->LoginStatus = ELoginStatus::LoggedIn;
 		AccountInfoRegistryOSSAdapter.Register(AccountInfoOSSAdapter);
@@ -318,12 +317,12 @@ TOnlineAsyncOpHandle<FAuthLogout> FAuthOSSAdapter::Logout(FAuthLogout::Params&&
 		}
 
 		// Set user auth data on operation.
-		InAsyncOp.Data.Set<TSharedRef<FAccountInfoOSSAdapter>>(AccountInfoKeyName, AccountInfoOSSAdapter.ToSharedRef());
+		InAsyncOp.Data.Set<TSharedRef<FAccountInfoOSSAdapter>>(UE_ONLINE_AUTH_OSS_ADAPTER_KEY_NAME_ACCOUNT_INFO, AccountInfoOSSAdapter.ToSharedRef());
 	})
 	// Step 2: Logout the user.
 	.Then([this](TOnlineAsyncOp<FAuthLogout>& InAsyncOp)
 	{
-		const TSharedRef<FAccountInfoOSSAdapter>& AccountInfoOSSAdapter = GetOpDataChecked<TSharedRef<FAccountInfoOSSAdapter>>(InAsyncOp, AccountInfoKeyName);
+		const TSharedRef<FAccountInfoOSSAdapter>& AccountInfoOSSAdapter = GetOpDataChecked<TSharedRef<FAccountInfoOSSAdapter>>(InAsyncOp, UE_ONLINE_AUTH_OSS_ADAPTER_KEY_NAME_ACCOUNT_INFO);
 
 		TPromise<void> Promise;
 		TFuture<void> Future = Promise.GetFuture();
@@ -343,7 +342,7 @@ TOnlineAsyncOpHandle<FAuthLogout> FAuthOSSAdapter::Logout(FAuthLogout::Params&&
 	// Step 3: bookkeeping and notifications.
 	.Then([this](TOnlineAsyncOp<FAuthLogout>& InAsyncOp)
 	{
-		const TSharedRef<FAccountInfoOSSAdapter>& AccountInfoOSSAdapter = GetOpDataChecked<TSharedRef<FAccountInfoOSSAdapter>>(InAsyncOp, AccountInfoKeyName);
+		const TSharedRef<FAccountInfoOSSAdapter>& AccountInfoOSSAdapter = GetOpDataChecked<TSharedRef<FAccountInfoOSSAdapter>>(InAsyncOp, UE_ONLINE_AUTH_OSS_ADAPTER_KEY_NAME_ACCOUNT_INFO);
 
 		UE_LOG(LogOnlineServices, Log, TEXT("[FAuthOSSAdapter::Logout][%s] Successfully logged out [%s]"), *GetSubsystem().GetSubsystemName().ToString(), *ToLogString(AccountInfoOSSAdapter->AccountId));
 		AccountInfoOSSAdapter->LoginStatus = ELoginStatus::NotLoggedIn;
@@ -374,7 +373,7 @@ TOnlineAsyncOpHandle<FAuthQueryExternalServerAuthTicket> FAuthOSSAdapter::QueryE
 		}
 
 		// Set translator traits on operation.
-		InAsyncOp.Data.Set<TSharedRef<FExternalServerAuthTicketTranslationTraits>>(ExternalServerAuthTicketTranslationTraitsKeyName, TranslationTraits.GetOkValue());
+		InAsyncOp.Data.Set<TSharedRef<FExternalServerAuthTicketTranslationTraits>>(UE_ONLINE_AUTH_OSS_ADAPTER_KEY_NAME_EXTERNAL_SERVER_AUTH_TICKET_TRANSLATION_TRAITS, TranslationTraits.GetOkValue());
 
 		// Look up logged in user.
 		TSharedPtr<FAccountInfoOSSAdapter> AccountInfoOSSAdapter = AccountInfoRegistryOSSAdapter.Find(Params.LocalAccountId);
@@ -385,13 +384,13 @@ TOnlineAsyncOpHandle<FAuthQueryExternalServerAuthTicket> FAuthOSSAdapter::QueryE
 		}
 
 		// Set user auth data on operation.
-		InAsyncOp.Data.Set<TSharedRef<FAccountInfoOSSAdapter>>(AccountInfoKeyName, AccountInfoOSSAdapter.ToSharedRef());
+		InAsyncOp.Data.Set<TSharedRef<FAccountInfoOSSAdapter>>(UE_ONLINE_AUTH_OSS_ADAPTER_KEY_NAME_ACCOUNT_INFO, AccountInfoOSSAdapter.ToSharedRef());
 	})
 	// Step 2: Fetch OSSv1 auth ticket data and signal result.
 	.Then([this](TOnlineAsyncOp<FAuthQueryExternalServerAuthTicket>& InAsyncOp)
 	{
-		const TSharedRef<FAccountInfoOSSAdapter>& AccountInfoOSSAdapter = GetOpDataChecked<TSharedRef<FAccountInfoOSSAdapter>>(InAsyncOp, AccountInfoKeyName);
-		const TSharedRef<FExternalServerAuthTicketTranslationTraits>& ExternalServerAuthTicketTranslationTraits = GetOpDataChecked<TSharedRef<FExternalServerAuthTicketTranslationTraits>>(InAsyncOp, ExternalAuthTokenTranslationTraitsKeyName);
+		const TSharedRef<FAccountInfoOSSAdapter>& AccountInfoOSSAdapter = GetOpDataChecked<TSharedRef<FAccountInfoOSSAdapter>>(InAsyncOp, UE_ONLINE_AUTH_OSS_ADAPTER_KEY_NAME_ACCOUNT_INFO);
+		const TSharedRef<FExternalServerAuthTicketTranslationTraits>& ExternalServerAuthTicketTranslationTraits = GetOpDataChecked<TSharedRef<FExternalServerAuthTicketTranslationTraits>>(InAsyncOp, UE_ONLINE_AUTH_OSS_ADAPTER_KEY_NAME_EXTERNAL_AUTH_TOKEN_TRANSLATION_TRAITS);
 
 		InAsyncOp.SetResult(FAuthQueryExternalServerAuthTicket::Result{FExternalServerAuthTicket{ExternalServerAuthTicketTranslationTraits->ExternalTicketType, GetIdentityInterface()->GetAuthToken(AccountInfoOSSAdapter->LocalUserNum)}});
 	})
@@ -418,7 +417,7 @@ TOnlineAsyncOpHandle<FAuthQueryExternalAuthToken> FAuthOSSAdapter::QueryExternal
 		}
 
 		// Set translator traits on operation.
-		InAsyncOp.Data.Set<const FExternalAuthTokenTranslationTraits*>(ExternalAuthTokenTranslationTraitsKeyName, TranslationTraits);
+		InAsyncOp.Data.Set<const FExternalAuthTokenTranslationTraits*>(UE_ONLINE_AUTH_OSS_ADAPTER_KEY_NAME_EXTERNAL_AUTH_TOKEN_TRANSLATION_TRAITS, TranslationTraits);
 
 		// Look up logged in user.
 		TSharedPtr<FAccountInfoOSSAdapter> AccountInfoOSSAdapter = AccountInfoRegistryOSSAdapter.Find(Params.LocalAccountId);
@@ -429,13 +428,13 @@ TOnlineAsyncOpHandle<FAuthQueryExternalAuthToken> FAuthOSSAdapter::QueryExternal
 		}
 
 		// Set user auth data on operation.
-		InAsyncOp.Data.Set<TSharedRef<FAccountInfoOSSAdapter>>(AccountInfoKeyName, AccountInfoOSSAdapter.ToSharedRef());
+		InAsyncOp.Data.Set<TSharedRef<FAccountInfoOSSAdapter>>(UE_ONLINE_AUTH_OSS_ADAPTER_KEY_NAME_ACCOUNT_INFO, AccountInfoOSSAdapter.ToSharedRef());
 	})
 	// Step 2: Fetch OSSv1 token data.
 	.Then([this](TOnlineAsyncOp<FAuthQueryExternalAuthToken>& InAsyncOp)
 	{
-		const TSharedRef<FAccountInfoOSSAdapter>& AccountInfoOSSAdapter = GetOpDataChecked<TSharedRef<FAccountInfoOSSAdapter>>(InAsyncOp, AccountInfoKeyName);
-		const FExternalAuthTokenTranslationTraits* ExternalAuthTokenTranslationTraits = GetOpDataChecked<const FExternalAuthTokenTranslationTraits*>(InAsyncOp, ExternalAuthTokenTranslationTraitsKeyName);
+		const TSharedRef<FAccountInfoOSSAdapter>& AccountInfoOSSAdapter = GetOpDataChecked<TSharedRef<FAccountInfoOSSAdapter>>(InAsyncOp, UE_ONLINE_AUTH_OSS_ADAPTER_KEY_NAME_ACCOUNT_INFO);
+		const FExternalAuthTokenTranslationTraits* ExternalAuthTokenTranslationTraits = GetOpDataChecked<const FExternalAuthTokenTranslationTraits*>(InAsyncOp, UE_ONLINE_AUTH_OSS_ADAPTER_KEY_NAME_EXTERNAL_AUTH_TOKEN_TRANSLATION_TRAITS);
 
 		TPromise<FString> Promise;
 		TFuture<FString> Future = Promise.GetFuture();
@@ -512,7 +511,7 @@ TOnlineAsyncOpHandle<FAuthQueryExternalAuthToken> FAuthOSSAdapter::QueryExternal
 	// Step 3: Signal valid result.
 	.Then([this](TOnlineAsyncOp<FAuthQueryExternalAuthToken>& InAsyncOp, FString&& ResolvedExternalAuthToken)
 	{
-		const FExternalAuthTokenTranslationTraits* ExternalAuthTokenTranslationTraits = GetOpDataChecked<const FExternalAuthTokenTranslationTraits*>(InAsyncOp, ExternalAuthTokenTranslationTraitsKeyName);
+		const FExternalAuthTokenTranslationTraits* ExternalAuthTokenTranslationTraits = GetOpDataChecked<const FExternalAuthTokenTranslationTraits*>(InAsyncOp, UE_ONLINE_AUTH_OSS_ADAPTER_KEY_NAME_EXTERNAL_AUTH_TOKEN_TRANSLATION_TRAITS);
 		InAsyncOp.SetResult(FAuthQueryExternalAuthToken::Result{FExternalAuthToken{ExternalAuthTokenTranslationTraits->ExternalLoginType, MoveTemp(ResolvedExternalAuthToken)}});
 	})
 	.Enqueue(GetSerialQueue());
@@ -597,13 +596,13 @@ TOnlineAsyncOpHandle<FAuthHandleLoginStatusChangedImpl> FAuthOSSAdapter::HandleL
 		}
 
 		// Set user auth data on operation.
-		InAsyncOp.Data.Set<TSharedRef<FAccountInfoOSSAdapter>>(AccountInfoKeyName, AccountInfoOSSAdapter.ToSharedRef());
+		InAsyncOp.Data.Set<TSharedRef<FAccountInfoOSSAdapter>>(UE_ONLINE_AUTH_OSS_ADAPTER_KEY_NAME_ACCOUNT_INFO, AccountInfoOSSAdapter.ToSharedRef());
 	})
 	// Step 2: Update status and notify.
 	.Then([this](TOnlineAsyncOp<FAuthHandleLoginStatusChangedImpl>& InAsyncOp)
 	{
 		const FAuthHandleLoginStatusChangedImpl::Params& Params = InAsyncOp.GetParams();
-		const TSharedRef<FAccountInfoOSSAdapter>& AccountInfoOSSAdapter = GetOpDataChecked<TSharedRef<FAccountInfoOSSAdapter>>(InAsyncOp, AccountInfoKeyName);
+		const TSharedRef<FAccountInfoOSSAdapter>& AccountInfoOSSAdapter = GetOpDataChecked<TSharedRef<FAccountInfoOSSAdapter>>(InAsyncOp, UE_ONLINE_AUTH_OSS_ADAPTER_KEY_NAME_ACCOUNT_INFO);
 
 		UE_LOG(LogOnlineServices, Log, TEXT("[FAuthOSSAdapter::HandleLoginStatusChangedImplOp][%s] Login status changed for account [%s]: %s => %s."),
 			*GetSubsystem().GetSubsystemName().ToString(), *ToLogString(AccountInfoOSSAdapter->AccountId), *ToLogString(AccountInfoOSSAdapter->LoginStatus), *ToLogString(Params.NewLoginStatus));
diff --git a/Engine/Plugins/Online/OnlineSubsystem/Source/Private/OnlineError.cpp b/Engine/Plugins/Online/OnlineSubsystem/Source/Private/OnlineError.cpp
--- a/Engine/Plugins/Online/OnlineSubsystem/Source/Private/OnlineError.cpp
+++ b/Engine/Plugins/Online/OnlineSubsystem/Source/Private/OnlineError.cpp
@@ -5,8 +5,10 @@
 #define LOCTEXT_NAMESPACE "OnlineError"
 #define ONLINE_ERROR_NAMESPACE "errors.com.epicgames.online.generic"
 
-//static
+PRAGMA_DISABLE_DEPRECATION_WARNINGS
 const FString FOnlineError::GenericErrorCode = TEXT("GenericError");
+PRAGMA_ENABLE_DEPRECATION_WARNINGS
+
 const FOnlineError& FOnlineError::Success() { static FOnlineError Error(EOnlineErrorResult::Success); return Error; }
 
 #if ONLINE_ERROR_LEGACY
@@ -138,11 +140,16 @@ void FOnlineError::SetFromErrorCode(FString&& ErrorCodeIn)
 void FOnlineError::SetFromErrorMessage(const FText& ErrorMessageIn)
 {
 	ErrorMessage = ErrorMessageIn;
-	ErrorCode = FTextInspector::GetKey(ErrorMessageIn).Get(GenericErrorCode);
+	ErrorCode = FTextInspector::GetKey(ErrorMessageIn).Get(GetGenericErrorCode());
 	ErrorRaw = ErrorMessageIn.ToString();
 	Result = EOnlineErrorResult::FailExtended;
 }
 
+const FString FOnlineError::GetGenericErrorCode()
+{
+	return TEXT("GenericError");
+}
+
 FString FOnlineError::ToLogString() const
 {
 	if (bSucceeded)
diff --git a/Engine/Plugins/Online/OnlineSubsystem/Source/Private/OnlineSubsystem.cpp b/Engine/Plugins/Online/OnlineSubsystem/Source/Private/OnlineSubsystem.cpp
--- a/Engine/Plugins/Online/OnlineSubsystem/Source/Private/OnlineSubsystem.cpp
+++ b/Engine/Plugins/Online/OnlineSubsystem/Source/Private/OnlineSubsystem.cpp
@@ -90,6 +90,7 @@ const FString DefaultSessionIdKey = TEXT("SessionId");
 /** Resource the client is logged in with */
 const FString PresenceResourceKey = TEXT("ResourceKey");
 
+PRAGMA_DISABLE_DEPRECATION_WARNINGS
 namespace OnlineIdentity
 {
 	namespace Errors
@@ -105,6 +106,7 @@ namespace OnlineIdentity
 		const FString NoAuthType = TEXT("no_auth_type");
 	}
 }
+PRAGMA_ENABLE_DEPRECATION_WARNINGS
 
 /** Workaround, please avoid using this */
 FUniqueNetIdPtr GetFirstSignedInUser(IOnlineIdentityPtr IdentityInt)
diff --git a/Engine/Plugins/Online/OnlineSubsystem/Source/Public/OnlineError.h b/Engine/Plugins/Online/OnlineSubsystem/Source/Public/OnlineError.h
--- a/Engine/Plugins/Online/OnlineSubsystem/Source/Public/OnlineError.h
+++ b/Engine/Plugins/Online/OnlineSubsystem/Source/Public/OnlineError.h
@@ -52,6 +52,8 @@ enum class EOnlineErrorResult : uint8
 	Canceled,
 	/** Extended error. More info can be found in the results or by looking at the ErrorCode */
 	FailExtended,
+	/** No game session found */
+	NoGameSession,
 
 	/** Default state */
 	Unknown
@@ -110,8 +112,12 @@ public:
 	inline FOnlineError& SetErrorRaw(const FString& Val) { ErrorRaw = Val; return *this; }
 
 	/** Code useful when all you have is raw error info from old APIs */
+	UE_DEPRECATED(5.2, "Access of GenericErrorCode is deprecated, use GetGenericErrorCode() instead.")
 	static const FString GenericErrorCode;
 
+	/** Code useful when all you have is raw error info from old APIs */
+	static const FString GetGenericErrorCode();
+
 	/** prints out everything, need something like this!!! */
 	FString GetErrorLegacy();
 	/** Call this if you want to log this out (will pick the best string representation) */
diff --git a/Engine/Plugins/Online/OnlineSubsystem/Source/Public/OnlineIdentityErrors.h b/Engine/Plugins/Online/OnlineSubsystem/Source/Public/OnlineIdentityErrors.h
--- a/Engine/Plugins/Online/OnlineSubsystem/Source/Public/OnlineIdentityErrors.h
+++ b/Engine/Plugins/Online/OnlineSubsystem/Source/Public/OnlineIdentityErrors.h
@@ -26,13 +26,20 @@ namespace OnlineIdentity
 		inline FOnlineError PinGrantTimeout() { return ONLINE_ERROR(EOnlineErrorResult::FailExtended, TEXT("pin_grant_timeout")); }
 
 		// Params
+		UE_DEPRECATED(5.2, "AuthLoginParam is deprecated, please use UE_ONLINE_ERROR_PARAM_IDENTITY_AUTH_LOGIN instead.")
 		extern ONLINESUBSYSTEM_API const FString AuthLoginParam;
+
+		UE_DEPRECATED(5.2, "AuthTypeParam is deprecated, please use UE_ONLINE_ERROR_PARAM_IDENTITY_AUTH_TYPE instead.")
 		extern ONLINESUBSYSTEM_API const FString AuthTypeParam;
+		UE_DEPRECATED(5.2, "AuthPasswordParam is deprecated, please use UE_ONLINE_ERROR_PARAM_IDENTITY_AUTH_PASSWORD instead.")
 		extern ONLINESUBSYSTEM_API const FString AuthPasswordParam;
 
 		// Results
+		UE_DEPRECATED(5.2, "NoUserId is deprecated, please use UE_ONLINE_ERROR_PARAM_IDENTITY_NO_USER_ID instead.")
 		extern ONLINESUBSYSTEM_API const FString NoUserId;
+		UE_DEPRECATED(5.2, "NoAuthToken is deprecated, please use UE_ONLINE_ERROR_PARAM_IDENTITY_NO_AUTH_TOKEN instead.")
 		extern ONLINESUBSYSTEM_API const FString NoAuthToken;
+		UE_DEPRECATED(5.2, "NoAuthType is deprecated, please use UE_ONLINE_ERROR_PARAM_IDENTITY_NO_AUTH_TYPE instead.")
 		extern ONLINESUBSYSTEM_API const FString NoAuthType;
 	}
 }
@@ -42,5 +49,11 @@ namespace OnlineIdentity
 #undef ONLINE_ERROR_NAMESPACE
 
 
+#define UE_ONLINE_ERROR_PARAM_IDENTITY_AUTH_LOGIN TEXT("auth_login")
+#define UE_ONLINE_ERROR_PARAM_IDENTITY_AUTH_TYPE TEXT("auth_type")
+#define UE_ONLINE_ERROR_PARAM_IDENTITY_AUTH_PASSWORD TEXT("auth_password")
 
+#define UE_ONLINE_ERROR_PARAM_IDENTITY_NO_USER_ID TEXT("no_user_id")
+#define UE_ONLINE_ERROR_PARAM_IDENTITY_NO_AUTH_TOKEN TEXT("no_auth_token")
+#define UE_ONLINE_ERROR_PARAM_IDENTITY_NO_AUTH_TYPE TEXT("no_auth_type")
 
