Refactor of EOS SendSessionInvite logic to fix promise lifetime crash

diff --git a/Engine/Plugins/Online/OnlineServicesEOSGS/Source/Private/Online/SessionsEOSGS.cpp b/Engine/Plugins/Online/OnlineServicesEOSGS/Source/Private/Online/SessionsEOSGS.cpp
--- a/Engine/Plugins/Online/OnlineServicesEOSGS/Source/Private/Online/SessionsEOSGS.cpp
+++ b/Engine/Plugins/Online/OnlineServicesEOSGS/Source/Private/Online/SessionsEOSGS.cpp
@@ -875,39 +875,37 @@ TFuture<TDefaultErrorResult<FUpdateSessionJoinabilityImpl>> FSessionsEOSGS::Upda
 	return Future;
 }
 
-TOnlineAsyncOpHandle<FSendSingleSessionInviteImpl> FSessionsEOSGS::SendSingleSessionInviteImpl(FSendSingleSessionInviteImpl::Params&& Params)
+TFuture<TDefaultErrorResult<FSendSingleSessionInviteImpl>> FSessionsEOSGS::SendSingleSessionInviteImpl(FSendSingleSessionInviteImpl::Params&& Params)
 {
-	TOnlineAsyncOpRef<FSendSingleSessionInviteImpl> Op = GetOp<FSendSingleSessionInviteImpl>(MoveTemp(Params));
-	const FSendSingleSessionInviteImpl::Params& OpParams = Op->GetParams();
+	TPromise<TDefaultErrorResult<FSendSingleSessionInviteImpl>> Promise;
+	TFuture<TDefaultErrorResult<FSendSingleSessionInviteImpl>> Future = Promise.GetFuture();
 
 	EOS_Sessions_SendInviteOptions SendInviteOptions = { };
 	SendInviteOptions.ApiVersion = 1;
 	UE_EOS_CHECK_API_MISMATCH(EOS_SESSIONS_SENDINVITE_API_LATEST, 1);
 
-	SendInviteOptions.LocalUserId = GetProductUserIdChecked(OpParams.LocalAccountId);
+	SendInviteOptions.LocalUserId = GetProductUserIdChecked(Params.LocalAccountId);
 
-	const FTCHARToUTF8 SessionNameUtf8(*OpParams.SessionName.ToString());
+	const FTCHARToUTF8 SessionNameUtf8(*Params.SessionName.ToString());
 	SendInviteOptions.SessionName = SessionNameUtf8.Get();
 
-	SendInviteOptions.TargetUserId = GetProductUserIdChecked(OpParams.TargetAccountId);
+	SendInviteOptions.TargetUserId = GetProductUserIdChecked(Params.TargetAccountId);
 
 	EOS_Async(EOS_Sessions_SendInvite, SessionsHandle, SendInviteOptions,
-		[this, WeakOp = Op->AsWeak()](const EOS_Sessions_SendInviteCallbackInfo* Result) mutable
+		[this, Promise = MoveTemp(Promise)](const EOS_Sessions_SendInviteCallbackInfo* Result) mutable
 	{
-		if (TOnlineAsyncOpPtr<FSendSingleSessionInviteImpl> StrongOp = WeakOp.Pin())
+		if (Result->ResultCode != EOS_EResult::EOS_Success)
 		{
-			if (Result->ResultCode != EOS_EResult::EOS_Success)
-			{
-				UE_LOG(LogTemp, Warning, TEXT("EOS_Sessions_SendInvite failed with result [%s]"), *LexToString(Result->ResultCode));
-				StrongOp->SetError(Errors::FromEOSResult(Result->ResultCode));
-				return;
-			}
-
-			StrongOp->SetResult(FSendSingleSessionInviteImpl::Result{ });
+			UE_LOG(LogTemp, Warning, TEXT("EOS_Sessions_SendInvite failed with result [%s]"), *LexToString(Result->ResultCode));
+			Promise.EmplaceValue(Errors::FromEOSResult(Result->ResultCode));
+		}
+		else
+		{
+			Promise.EmplaceValue(FSendSingleSessionInviteImpl::Result{ });
 		}
 	});
 
-	return Op->GetHandle();
+	return Future;
 }
 
 TFuture<TOnlineResult<FLeaveSession>> FSessionsEOSGS::LeaveSessionImpl(const FLeaveSession::Params& Params)
@@ -1419,21 +1417,21 @@ TFuture<TOnlineResult<FSendSessionInvite>> FSessionsEOSGS::SendSessionInviteImpl
 		FSendSingleSessionInviteParams.SessionName = Params.SessionName;
 		FSendSingleSessionInviteParams.TargetAccountId = TargetAccountId;
 
-		TSharedRef<TPromise<TDefaultErrorResult<FSendSingleSessionInviteImpl>>> SessionInvitePromise = MakeShared<TPromise<TDefaultErrorResult<FSendSingleSessionInviteImpl>>>();
+		TPromise<TDefaultErrorResult<FSendSingleSessionInviteImpl>> SessionInvitePromise;
+		PendingSessionInvites.Emplace(SessionInvitePromise.GetFuture());
+
 		SendSingleSessionInviteImpl(MoveTemp(FSendSingleSessionInviteParams))
-			.OnComplete([SessionInvitePromise](const TOnlineResult<FSendSingleSessionInviteImpl>& Result)
+			.Next([SessionInvitePromise = MoveTemp(SessionInvitePromise)](TDefaultErrorResult<FSendSingleSessionInviteImpl>&& Result) mutable
 				{
 					if (Result.IsOk())
 					{
-						SessionInvitePromise->EmplaceValue(Result.GetOkValue());
+						SessionInvitePromise.EmplaceValue(MoveTemp(Result.GetOkValue()));
 					}
 					else
 					{
-						SessionInvitePromise->EmplaceValue(Result.GetErrorValue());
+						SessionInvitePromise.EmplaceValue(MoveTemp(Result.GetErrorValue()));
 					}
 				});
-
-		PendingSessionInvites.Emplace(SessionInvitePromise->GetFuture());
 	}
 
 	WhenAll(MoveTemp(PendingSessionInvites))
diff --git a/Engine/Plugins/Online/OnlineServicesEOSGS/Source/Public/Online/SessionsEOSGS.h b/Engine/Plugins/Online/OnlineServicesEOSGS/Source/Public/Online/SessionsEOSGS.h
--- a/Engine/Plugins/Online/OnlineServicesEOSGS/Source/Public/Online/SessionsEOSGS.h
+++ b/Engine/Plugins/Online/OnlineServicesEOSGS/Source/Public/Online/SessionsEOSGS.h
@@ -254,7 +254,7 @@ protected:
 	/**
 	 * Internal method called by SendSessionInvite for every user
 	 */
-	TOnlineAsyncOpHandle<FSendSingleSessionInviteImpl> SendSingleSessionInviteImpl(FSendSingleSessionInviteImpl::Params&& Params);
+	TFuture<TDefaultErrorResult<FSendSingleSessionInviteImpl>> SendSingleSessionInviteImpl(FSendSingleSessionInviteImpl::Params&& Params);
 
 	static FOnlineSessionId CreateSessionId(const FString& SessionId);
 	FSessionInviteId CreateSessionInviteId(const FString& SessionInviteId) const;
