Add stat type and default value into the stat definition.

diff --git a/Engine/Plugins/Online/OnlineServices/Source/OnlineServicesCommon/Public/Online/StatsCommon.h b/Engine/Plugins/Online/OnlineServices/Source/OnlineServicesCommon/Public/Online/StatsCommon.h
--- a/Engine/Plugins/Online/OnlineServices/Source/OnlineServicesCommon/Public/Online/StatsCommon.h
+++ b/Engine/Plugins/Online/OnlineServices/Source/OnlineServicesCommon/Public/Online/StatsCommon.h
@@ -30,6 +30,8 @@ struct FStatDefinition
 	int32 Id = 0;
 	/* How the stat will be modified */
 	EStatModifyMethod ModifyMethod = EStatModifyMethod::Set;
+	/* Store the default type and value */
+	FStatValue DefaultValue = FStatValue(int64(0));
 };
 
 struct FStatsCommonConfig
@@ -42,7 +44,8 @@ namespace Meta {
 BEGIN_ONLINE_STRUCT_META(FStatDefinition)
 	ONLINE_STRUCT_FIELD(FStatDefinition, Name),
 	ONLINE_STRUCT_FIELD(FStatDefinition, Id),
-	ONLINE_STRUCT_FIELD(FStatDefinition, ModifyMethod)
+	ONLINE_STRUCT_FIELD(FStatDefinition, ModifyMethod),
+	ONLINE_STRUCT_FIELD(FStatDefinition, DefaultValue)
 END_ONLINE_STRUCT_META()
 
 BEGIN_ONLINE_STRUCT_META(FStatsCommonConfig)
diff --git a/Engine/Plugins/Online/OnlineServicesEOSGS/Source/Private/Online/StatsEOSGS.cpp b/Engine/Plugins/Online/OnlineServicesEOSGS/Source/Private/Online/StatsEOSGS.cpp
--- a/Engine/Plugins/Online/OnlineServicesEOSGS/Source/Private/Online/StatsEOSGS.cpp
+++ b/Engine/Plugins/Online/OnlineServicesEOSGS/Source/Private/Online/StatsEOSGS.cpp
@@ -9,13 +9,14 @@
 
 #include "eos_stats.h"
 
-namespace UE::Online {
+#define UE_ONLINE_STAT_EOS_DOUBLE_PRECISION (0.001) // NOTE!!! Modifying the precision will result in broken value for existing stats of all users!
+#define UE_ONLINE_STAT_EOS_NAME_MAX_LENGTH (256 + 1) // 256 plus null terminator
 
-#define STAT_NAME_MAX_LENGTH_EOS 256 + 1 // 256 plus null terminator
+namespace UE::Online {
 
 struct FStatNameRawBufferEOS
 {
-	char StatName[STAT_NAME_MAX_LENGTH_EOS];
+	char StatName[UE_ONLINE_STAT_EOS_NAME_MAX_LENGTH];
 };
 
 FStatsEOSGS::FStatsEOSGS(FOnlineServicesEOSGS& InServices)
@@ -69,8 +70,38 @@ TOnlineAsyncOpHandle<FUpdateStats> FStatsEOSGS::UpdateStats(FUpdateStats::Params
 				EOSStat.ApiVersion = 1;
 				UE_EOS_CHECK_API_MISMATCH(EOS_STATS_INGESTDATA_API_LATEST, 1);
 
-				EOSStat.IngestAmount = static_cast<int32>(StatPair.Value.GetInt64());
-				FCStringAnsi::Strncpy(EOSStatNames[Index].StatName, TCHAR_TO_UTF8(*StatPair.Key), STAT_NAME_MAX_LENGTH_EOS);
+				const FStatDefinition* StatDefinition = GetStatDefinition(StatPair.Key);
+				check(StatDefinition);
+				check(StatPair.Value.GetType() == StatDefinition->DefaultValue.GetType());
+
+				switch (StatPair.Value.GetType())
+				{
+				case ESchemaAttributeType::Bool:
+					EOSStat.IngestAmount = StatPair.Value.GetBoolean();
+					break;
+				case ESchemaAttributeType::Double:
+				{
+					double ConvertedDouble = StatPair.Value.GetDouble() / UE_ONLINE_STAT_EOS_DOUBLE_PRECISION;
+					double MinLimit = static_cast<double>(TNumericLimits<int32>::Min());
+					double MaxLimit = static_cast<double>(TNumericLimits<int32>::Max());
+					check(ConvertedDouble > MinLimit && ConvertedDouble < MaxLimit);
+					EOSStat.IngestAmount = static_cast<int32>(FMath::Clamp(ConvertedDouble, MinLimit, MaxLimit));
+					break;
+				}
+				case ESchemaAttributeType::Int64:
+				{
+					int64 MinLimit = static_cast<int64>(TNumericLimits<int32>::Min());
+					int64 MaxLimit = static_cast<int64>(TNumericLimits<int32>::Max());
+					check(StatPair.Value.GetInt64() > MinLimit && StatPair.Value.GetInt64() < MaxLimit);
+					EOSStat.IngestAmount = static_cast<int32>(FMath::Clamp(StatPair.Value.GetInt64(), MinLimit, MaxLimit));
+					break;
+				}
+				default: checkNoEntry(); // Intentional fallthrough
+				case ESchemaAttributeType::String: // Not supported
+					break;
+				}
+
+				FCStringAnsi::Strncpy(EOSStatNames[Index].StatName, TCHAR_TO_UTF8(*StatPair.Key), UE_ONLINE_STAT_EOS_NAME_MAX_LENGTH);
 				EOSStat.StatName = EOSStatNames[Index].StatName;
 
 				Index++;
@@ -128,7 +159,7 @@ void QueryStatsEOS(EOS_HStats StatsHandle, const FAccountId& LocalAccountId, con
 	uint32 Index = 0;
 	for (const FString& StatName : StatNames)
 	{
-		FCStringAnsi::Strncpy(EOSStatNames[Index].StatName, TCHAR_TO_UTF8(*StatName), STAT_NAME_MAX_LENGTH_EOS);
+		FCStringAnsi::Strncpy(EOSStatNames[Index].StatName, TCHAR_TO_UTF8(*StatName), UE_ONLINE_STAT_EOS_NAME_MAX_LENGTH);
 		EOSStatNamesPtr[Index] = EOSStatNames[Index].StatName;
 		Index++;
 	}
@@ -141,21 +172,43 @@ void QueryStatsEOS(EOS_HStats StatsHandle, const FAccountId& LocalAccountId, con
 	EOS_Async(EOS_Stats_QueryStats, StatsHandle, Options, MoveTemp(Promise));
 }
 
-void ReadStatsFromEOSResult(EOS_HStats StatsHandle, const EOS_Stats_OnQueryStatsCompleteCallbackInfo* Data, const TArray<FString>& StatNames, TMap<FString, FStatValue>& OutStats)
+}
+
+void FStatsEOSGS::ReadStatsFromEOSResult(const EOS_Stats_OnQueryStatsCompleteCallbackInfo* Data, const TArray<FString>& StatNames, TMap<FString, FStatValue>& OutStats)
 {
 	for (const FString& StatName : StatNames)
 	{
-		char StatNameANSI[STAT_NAME_MAX_LENGTH_EOS];
+		char StatNameANSI[UE_ONLINE_STAT_EOS_NAME_MAX_LENGTH];
 		EOS_Stats_CopyStatByNameOptions Options = { };
 		Options.ApiVersion = 1;
 		UE_EOS_CHECK_API_MISMATCH(EOS_STATS_COPYSTATBYNAME_API_LATEST, 1);
 		Options.TargetUserId = Data->TargetUserId;
 		Options.Name = StatNameANSI;
-		FCStringAnsi::Strncpy(StatNameANSI, TCHAR_TO_UTF8(*StatName), STAT_NAME_MAX_LENGTH_EOS);
+		FCStringAnsi::Strncpy(StatNameANSI, TCHAR_TO_UTF8(*StatName), UE_ONLINE_STAT_EOS_NAME_MAX_LENGTH);
 
 		EOS_Stats_Stat* ReadStat = nullptr;
+
+		const FStatDefinition* StatDefinition = GetStatDefinition(StatName);
+		check(StatDefinition);
+
+		FStatValue StatValue = StatDefinition->DefaultValue;
 		if (EOS_Stats_CopyStatByName(StatsHandle, &Options, &ReadStat) == EOS_EResult::EOS_Success)
 		{
+			switch (StatDefinition->DefaultValue.GetType())
+			{
+				case ESchemaAttributeType::Bool:
+					StatValue.Set(static_cast<bool>(ReadStat->Value != 0));
+					break;
+				case ESchemaAttributeType::Double:
+					StatValue.Set(static_cast<double>(ReadStat->Value) * UE_ONLINE_STAT_EOS_DOUBLE_PRECISION);
+					break;
+				case ESchemaAttributeType::Int64:
+					StatValue.Set(static_cast<int64>(ReadStat->Value));
+					break;
+				default: checkNoEntry(); // Intentional fallthrough
+				case ESchemaAttributeType::String: // Not supported
+					break;
+			}
 			OutStats.Add(StatName, FStatValue(static_cast<int64>(ReadStat->Value)));
 			EOS_Stats_Stat_Release(ReadStat);
 		}
@@ -163,13 +216,11 @@ void ReadStatsFromEOSResult(EOS_HStats StatsHandle, const EOS_Stats_OnQueryStats
 		{
 			// Put an empty stat in
 			UE_LOG(LogTemp, VeryVerbose, TEXT("Value not found for stat %s, adding empty value"), *StatName);
-			OutStats.Add(StatName, FStatValue(static_cast<int64>(0)));
 		}
+		OutStats.Add(StatName, MoveTemp(StatValue));
 	}
 }
 
-}
-
 TOnlineAsyncOpHandle<FQueryStats> FStatsEOSGS::QueryStats(FQueryStats::Params&& Params)
 {
 	TOnlineAsyncOpRef<FQueryStats> Op = GetJoinableOp<FQueryStats>(MoveTemp(Params));
@@ -200,7 +251,7 @@ TOnlineAsyncOpHandle<FQueryStats> FStatsEOSGS::QueryStats(FQueryStats::Params&&
 			const FQueryStats::Params& Params = InAsyncOp.GetParams();
 
 			FQueryStats::Result Result;
-			Private::ReadStatsFromEOSResult(StatsHandle, Data, Params.StatNames, Result.Stats);
+			ReadStatsFromEOSResult(Data, Params.StatNames, Result.Stats);
 
 			FUserStats UserStats;
 			UserStats.AccountId = Params.TargetAccountId;
@@ -251,7 +302,7 @@ TOnlineAsyncOpHandle<FBatchQueryStats> FStatsEOSGS::BatchQueryStats(FBatchQueryS
 				FUserStats& UserStats = BatchQueriedUsersStats.Emplace_GetRef();
 				UserStats.AccountId = TargetAccountId;
 
-				Private::ReadStatsFromEOSResult(StatsHandle, Data, InAsyncOp.GetParams().StatNames, UserStats.Stats);
+				ReadStatsFromEOSResult(Data, InAsyncOp.GetParams().StatNames, UserStats.Stats);
 				CacheUserStats(UserStats);
 			});
 		}
diff --git a/Engine/Plugins/Online/OnlineServicesEOSGS/Source/Public/Online/StatsEOSGS.h b/Engine/Plugins/Online/OnlineServicesEOSGS/Source/Public/Online/StatsEOSGS.h
--- a/Engine/Plugins/Online/OnlineServicesEOSGS/Source/Public/Online/StatsEOSGS.h
+++ b/Engine/Plugins/Online/OnlineServicesEOSGS/Source/Public/Online/StatsEOSGS.h
@@ -14,6 +14,12 @@ namespace UE::Online {
 
 class FOnlineServicesEOSGS;
 
+/**
+ * Because the value type is int32 in EOS, stats types are converted to int32. So these are limitations: 
+ *	- String type of Stats is not supported;
+ *	- double type of Stats is casted to int32, and precision is UE_ONLINE_STAT_EOS_DOUBLE_PRECISION. Out of range value will be clamped;
+ *	- int64 type of Stats is casted to int32. Out of range value will be clamped;
+ */
 class ONLINESERVICESEOSGS_API FStatsEOSGS : public FStatsCommon
 {
 public:
@@ -29,7 +35,9 @@ public:
 	virtual TOnlineAsyncOpHandle<FQueryStats> QueryStats(FQueryStats::Params&& Params) override;
 	virtual TOnlineAsyncOpHandle<FBatchQueryStats> BatchQueryStats(FBatchQueryStats::Params&& Params) override;
 
-private:
+protected:
+	void ReadStatsFromEOSResult(const EOS_Stats_OnQueryStatsCompleteCallbackInfo* Data, const TArray<FString>& StatNames, TMap<FString, FStatValue>& OutStats);
+
 	EOS_HStats StatsHandle = nullptr;
 	TArray<FUserStats> BatchQueriedUsersStats;
 };
diff --git a/Engine/Plugins/Online/OnlineServicesNull/Source/Private/Online/StatsNull.cpp b/Engine/Plugins/Online/OnlineServicesNull/Source/Private/Online/StatsNull.cpp
--- a/Engine/Plugins/Online/OnlineServicesNull/Source/Private/Online/StatsNull.cpp
+++ b/Engine/Plugins/Online/OnlineServicesNull/Source/Private/Online/StatsNull.cpp
@@ -94,10 +94,7 @@ TOnlineAsyncOpHandle<FUpdateStats> FStatsNull::UpdateStats(FUpdateStats::Params&
 	return Op->GetHandle();
 }
 
-namespace StatsNullPrivate
-{
-
-void ReadStatsFromCache(const FUserStats* ExistingUserStats, const TArray<FString>& StatNames, TMap<FString, FStatValue>& OutStats)
+void FStatsNull::ReadStatsFromCache(const FUserStats* ExistingUserStats, const TArray<FString>& StatNames, TMap<FString, FStatValue>& OutStats)
 {
 	for (const FString& StatName : StatNames)
 	{
@@ -110,12 +107,12 @@ void ReadStatsFromCache(const FUserStats* ExistingUserStats, const TArray<FStrin
 			}
 		}
 
-		OutStats.Add(StatName, FStatValue(static_cast<int64>(0)));
+		const FStatDefinition* StatDefinition = GetStatDefinition(StatName);
+		check(StatDefinition);
+		OutStats.Add(StatName, StatDefinition->DefaultValue);
 	}
 }
 
-}
-
 TOnlineAsyncOpHandle<FQueryStats> FStatsNull::QueryStats(FQueryStats::Params&& Params)
 {
 	TOnlineAsyncOpRef<FQueryStats> Op = GetOp<FQueryStats>(MoveTemp(Params));
@@ -131,7 +128,7 @@ TOnlineAsyncOpHandle<FQueryStats> FStatsNull::QueryStats(FQueryStats::Params&& P
 		FQueryStats::Result Result;
 
 		FUserStats* ExistingUserStats = UsersStats.FindByPredicate(FFindUserStatsByAccountId(InAsyncOp.GetParams().TargetAccountId));
-		StatsNullPrivate::ReadStatsFromCache(ExistingUserStats, InAsyncOp.GetParams().StatNames, Result.Stats);
+		ReadStatsFromCache(ExistingUserStats, InAsyncOp.GetParams().StatNames, Result.Stats);
 
 		InAsyncOp.SetResult(MoveTemp(Result));
 	})
@@ -158,7 +155,7 @@ TOnlineAsyncOpHandle<FBatchQueryStats> FStatsNull::BatchQueryStats(FBatchQuerySt
 			FUserStats& UserStats = Result.UsersStats.Emplace_GetRef();
 			UserStats.AccountId = TargetAccountId;
 			FUserStats* ExistingUserStats = UsersStats.FindByPredicate(FFindUserStatsByAccountId(TargetAccountId));
-			StatsNullPrivate::ReadStatsFromCache(ExistingUserStats, InAsyncOp.GetParams().StatNames, UserStats.Stats);
+			ReadStatsFromCache(ExistingUserStats, InAsyncOp.GetParams().StatNames, UserStats.Stats);
 		}
 
 		InAsyncOp.SetResult(MoveTemp(Result));
diff --git a/Engine/Plugins/Online/OnlineServicesNull/Source/Public/Online/StatsNull.h b/Engine/Plugins/Online/OnlineServicesNull/Source/Public/Online/StatsNull.h
--- a/Engine/Plugins/Online/OnlineServicesNull/Source/Public/Online/StatsNull.h
+++ b/Engine/Plugins/Online/OnlineServicesNull/Source/Public/Online/StatsNull.h
@@ -25,6 +25,8 @@ public:
 #endif // !UE_BUILD_SHIPPING
 
 protected:
+	void ReadStatsFromCache(const FUserStats* ExistingUserStats, const TArray<FString>& StatNames, TMap<FString, FStatValue>& OutStats);
+
 	// TODO: Save UsersStats into local user profile
 	TArray<FUserStats> UsersStats;
 };
