Added LobbyId override option to CreateLobby method in OnlineServices Lobbies interface.
Added logic to LobbiesEOSGS to use the new option.

diff --git a/Engine/Plugins/Online/OnlineServicesEOSGS/Source/Private/Online/LobbiesEOSGS.cpp b/Engine/Plugins/Online/OnlineServicesEOSGS/Source/Private/Online/LobbiesEOSGS.cpp
--- a/Engine/Plugins/Online/OnlineServicesEOSGS/Source/Private/Online/LobbiesEOSGS.cpp
+++ b/Engine/Plugins/Online/OnlineServicesEOSGS/Source/Private/Online/LobbiesEOSGS.cpp
@@ -97,6 +97,10 @@ TOnlineAsyncOpHandle<FCreateLobby> FLobbiesEOSGS::CreateLobby(FCreateLobby::Para
 		return Future;
 	};
 
+	FSchemaVariant Attribute;
+	InParams.Attributes.RemoveAndCopyValue(LOBBIESEOSGS_LOBBY_ID_OVERRIDE_ATTRIBUTE_KEY, Attribute);
+	FString LobbyIdOverrideStr = Attribute.VariantType == ESchemaAttributeType::String ? Attribute.GetString() : FString();
+
 	TOnlineAsyncOpRef<FCreateLobby> Op = GetOp<FCreateLobby>(MoveTemp(InParams));
 	const FCreateLobby::Params& Params = Op->GetParams();
 
@@ -118,7 +122,7 @@ TOnlineAsyncOpHandle<FCreateLobby> FLobbiesEOSGS::CreateLobby(FCreateLobby::Para
 
 	// Start operation.
 	// Step 1: Call create lobby.
-	Op->Then([this](TOnlineAsyncOp<FCreateLobby>& InAsyncOp, TPromise<const EOS_Lobby_CreateLobbyCallbackInfo*>&& Promise)
+	Op->Then([this, LobbyIdOverrideStr](TOnlineAsyncOp<FCreateLobby>& InAsyncOp, TPromise<const EOS_Lobby_CreateLobbyCallbackInfo*>&& Promise)
 	{
 		const FCreateLobby::Params& Params = InAsyncOp.GetParams();
 		const FLobbyBucketIdTranslator<ELobbyTranslationType::ToService> BucketTanslator(LobbyPrerequisites->BucketId);
@@ -139,6 +143,9 @@ TOnlineAsyncOpHandle<FCreateLobby> FLobbiesEOSGS::CreateLobby(FCreateLobby::Para
 		CreateLobbyOptions.bDisableHostMigration = false; // todo: handle
 		CreateLobbyOptions.bEnableRTCRoom = false; // todo: handle
 
+		const auto LobbyIdOverrideConverter = StringCast<UTF8CHAR>(*LobbyIdOverrideStr);
+		CreateLobbyOptions.LobbyId = LobbyIdOverrideConverter.Length() ? (const char*)LobbyIdOverrideConverter.Get() : nullptr;
+		
 		EOS_Async(EOS_Lobby_CreateLobby, LobbyPrerequisites->LobbyInterfaceHandle, CreateLobbyOptions, MoveTemp(Promise)); 
 	})
 	// Step 2: Handle errors and attach the lobby id to the operation data.
diff --git a/Engine/Plugins/Online/OnlineServicesEOSGS/Source/Public/Online/OnlineServicesEOSGSTypes.h b/Engine/Plugins/Online/OnlineServicesEOSGS/Source/Public/Online/OnlineServicesEOSGSTypes.h
--- a/Engine/Plugins/Online/OnlineServicesEOSGS/Source/Public/Online/OnlineServicesEOSGSTypes.h
+++ b/Engine/Plugins/Online/OnlineServicesEOSGS/Source/Public/Online/OnlineServicesEOSGSTypes.h
@@ -12,6 +12,8 @@ namespace UE::Online {
 
 namespace Private {
 
+#define LOBBIESEOSGS_LOBBY_ID_OVERRIDE_ATTRIBUTE_KEY TEXT("LOBBY_ID_OVERRIDE")
+
 /** Class to handle all callbacks generically using a future to forward callback results */
 template<typename CallbackType>
 class TEOSCallback
