PR #10225: Fix OnlineServiceEOSGS bIsServer set to false in server mode
(Contributed by [REDACTED_PERSON])

diff --git a/Engine/Plugins/Online/OnlineServicesEOSGS/Source/Private/Online/OnlineServicesEOSGSPlatformFactory.cpp b/Engine/Plugins/Online/OnlineServicesEOSGS/Source/Private/Online/OnlineServicesEOSGSPlatformFactory.cpp
--- a/Engine/Plugins/Online/OnlineServicesEOSGS/Source/Private/Online/OnlineServicesEOSGSPlatformFactory.cpp
+++ b/Engine/Plugins/Online/OnlineServicesEOSGS/Source/Private/Online/OnlineServicesEOSGSPlatformFactory.cpp
@@ -102,7 +102,7 @@ IEOSPlatformHandlePtr FOnlineServicesEOSGSPlatformFactory::CreatePlatform()
 	PlatformOptions.ApiVersion = 12;
 	UE_EOS_CHECK_API_MISMATCH(EOS_PLATFORM_OPTIONS_API_LATEST, 12);
 	PlatformOptions.Reserved = nullptr;
-	PlatformOptions.bIsServer = EOS_FALSE;
+	PlatformOptions.bIsServer = IsRunningDedicatedServer() ? EOS_TRUE : EOS_FALSE;
 	PlatformOptions.OverrideCountryCode = nullptr;
 	PlatformOptions.OverrideLocaleCode = nullptr;
 	// Can't check GIsEditor here because it is too soon!
