PR #9975:
Fixed error paths in session search method not clearing cache maps
(Bug reported by [REDACTED_PERSON])

diff --git a/Engine/Plugins/Online/OnlineServicesEOSGS/Source/Private/Online/SessionsEOSGS.cpp b/Engine/Plugins/Online/OnlineServicesEOSGS/Source/Private/Online/SessionsEOSGS.cpp
--- a/Engine/Plugins/Online/OnlineServicesEOSGS/Source/Private/Online/SessionsEOSGS.cpp
+++ b/Engine/Plugins/Online/OnlineServicesEOSGS/Source/Private/Online/SessionsEOSGS.cpp
@@ -1109,6 +1109,8 @@ TFuture<TOnlineResult<FFindSessions>> FSessionsEOSGS::FindSessionsImpl(const FFi
 			{
 				UE_LOG(LogTemp, Warning, TEXT("EOS_SessionSearch_Find failed with result [%s]"), *LexToString(FindCallbackInfoResult->ResultCode));
 				Promise.EmplaceValue(Errors::FromEOSResult(FindCallbackInfoResult->ResultCode));
+				CurrentSessionSearchPromisesUserMap.Remove(Params.LocalAccountId);
+				CurrentSessionSearchHandleEOSGSUserMap.Remove(Params.LocalAccountId);
 				return;
 			}
 
@@ -1172,6 +1174,8 @@ TFuture<TOnlineResult<FFindSessions>> FSessionsEOSGS::FindSessionsImpl(const FFi
 					{
 						// Store first encountered error to return as result.
 						Promise.EmplaceValue(MoveTemp(Result.GetErrorValue()));
+						CurrentSessionSearchPromisesUserMap.Remove(Params.LocalAccountId);
+						CurrentSessionSearchHandleEOSGSUserMap.Remove(Params.LocalAccountId);
 						return;
 					}
 				}
