- Add "GetRegistryPriority" to EOS/EOSGS so we don't have to hardcode magic numbers, each derived OSS can +1 the value from the base.

diff --git a/Engine/Plugins/Online/OnlineServicesEOS/Source/Private/Online/OnlineServicesEOSModule.cpp b/Engine/Plugins/Online/OnlineServicesEOS/Source/Private/Online/OnlineServicesEOSModule.cpp
new file mode 100644
--- /dev/null
+++ b/Engine/Plugins/Online/OnlineServicesEOS/Source/Private/Online/OnlineServicesEOSModule.cpp
@@ -0,0 +1,57 @@
+// Copyright Epic Games, Inc. All Rights Reserved.
+
+#include "Online/OnlineServicesEOSModule.h"
+
+#include "Modules/ModuleManager.h"
+#include "Online/OnlineExecHandler.h"
+#include "Online/OnlineIdEOS.h"
+#include "Online/OnlineServicesEOS.h"
+#include "Online/OnlineServicesEOSGSModule.h"
+#include "Online/OnlineServicesRegistry.h"
+
+namespace UE::Online
+{
+
+class FOnlineServicesFactoryEOS : public IOnlineServicesFactory
+{
+public:
+	virtual ~FOnlineServicesFactoryEOS() {}
+	virtual TSharedPtr<IOnlineServices> Create(FName InInstanceName) override
+	{
+		return MakeShared<FOnlineServicesEOS>(InInstanceName);
+	}
+};
+
+int FOnlineServicesEOSModule::GetRegistryPriority()
+{
+	return FOnlineServicesEOSGSModule::GetRegistryPriority() + 1;
+}
+
+void FOnlineServicesEOSModule::StartupModule()
+{
+	const FName EOSSharedModuleName = TEXT("EOSShared");
+	if (!FModuleManager::Get().IsModuleLoaded(EOSSharedModuleName))
+	{
+		FModuleManager::Get().LoadModuleChecked(EOSSharedModuleName);
+	}
+
+	// Making sure we load the module at this point will avoid errors while cooking
+	const FName OnlineServicesInterfaceModuleName = TEXT("OnlineServicesInterface");
+	if (!FModuleManager::Get().IsModuleLoaded(OnlineServicesInterfaceModuleName))
+	{
+		FModuleManager::Get().LoadModuleChecked(OnlineServicesInterfaceModuleName);
+	}
+
+	FOnlineServicesRegistry::Get().RegisterServicesFactory(EOnlineServices::Epic, MakeUnique<FOnlineServicesFactoryEOS>(), GetRegistryPriority());
+	FOnlineIdRegistryRegistry::Get().RegisterAccountIdRegistry(EOnlineServices::Epic, &FOnlineAccountIdRegistryEOS::Get(), GetRegistryPriority());
+}
+
+void FOnlineServicesEOSModule::ShutdownModule()
+{
+	FOnlineServicesRegistry::Get().UnregisterServicesFactory(EOnlineServices::Epic, GetRegistryPriority());
+	FOnlineIdRegistryRegistry::Get().UnregisterAccountIdRegistry(EOnlineServices::Epic, GetRegistryPriority());
+}
+
+/* UE::Online */ }
+
+IMPLEMENT_MODULE(UE::Online::FOnlineServicesEOSModule, OnlineServicesEOS);
diff --git a/Engine/Plugins/Online/OnlineServicesEOS/Source/Private/OnlineServicesEOSModule.cpp b/Engine/Plugins/Online/OnlineServicesEOS/Source/Private/OnlineServicesEOSModule.cpp
deleted file mode 100644
diff --git a/Engine/Plugins/Online/OnlineServicesEOS/Source/Public/Online/OnlineServicesEOSModule.h b/Engine/Plugins/Online/OnlineServicesEOS/Source/Public/Online/OnlineServicesEOSModule.h
new file mode 100644
--- /dev/null
+++ b/Engine/Plugins/Online/OnlineServicesEOS/Source/Public/Online/OnlineServicesEOSModule.h
@@ -0,0 +1,21 @@
+// Copyright Epic Games, Inc. All Rights Reserved.
+
+#pragma once
+
+#include "Modules/ModuleInterface.h"
+
+namespace UE::Online {
+
+class FOnlineServicesEOSModule : public IModuleInterface
+{
+public:
+	ONLINESERVICESEOS_API static int GetRegistryPriority();
+
+private:
+	// Begin IModuleInterface
+	virtual void StartupModule() override;
+	virtual void ShutdownModule() override;
+	// End IModuleInterface
+};
+
+} // namespace UE::Online
\ No newline at end of file
diff --git a/Engine/Plugins/Online/OnlineServicesEOSGS/Source/Private/Online/OnlineServicesEOSGSModule.cpp b/Engine/Plugins/Online/OnlineServicesEOSGS/Source/Private/Online/OnlineServicesEOSGSModule.cpp
new file mode 100644
--- /dev/null
+++ b/Engine/Plugins/Online/OnlineServicesEOSGS/Source/Private/Online/OnlineServicesEOSGSModule.cpp
@@ -0,0 +1,60 @@
+// Copyright Epic Games, Inc. All Rights Reserved.
+
+#include "Online/OnlineServicesEOSGSModule.h"
+
+#include "Modules/ModuleManager.h"
+#include "Online/OnlineIdEOSGS.h"
+#include "Online/OnlineServicesEOSGS.h"
+#include "Online/OnlineServicesRegistry.h"
+#include "Online/OnlineServicesEOSGSPlatformFactory.h"
+#include "Online/SessionsEOSGS.h"
+
+namespace UE::Online
+{
+
+class FOnlineServicesFactoryEOSGS : public IOnlineServicesFactory
+{
+public:
+	virtual ~FOnlineServicesFactoryEOSGS() {}
+	virtual TSharedPtr<IOnlineServices> Create(FName InInstanceName) override
+	{
+		return MakeShared<FOnlineServicesEOSGS>(InInstanceName);
+	}
+};
+
+int FOnlineServicesEOSGSModule::GetRegistryPriority()
+{
+	return 0;
+}
+
+void FOnlineServicesEOSGSModule::StartupModule()
+{
+	// Making sure we load the module at this point will avoid errors while cooking
+	const FName OnlineServicesInterfaceModuleName = TEXT("OnlineServicesInterface");
+	if (!FModuleManager::Get().IsModuleLoaded(OnlineServicesInterfaceModuleName))
+	{
+		FModuleManager::Get().LoadModuleChecked(OnlineServicesInterfaceModuleName);
+	}
+
+	FOnlineServicesRegistry::Get().RegisterServicesFactory(EOnlineServices::Epic, MakeUnique<FOnlineServicesFactoryEOSGS>(), GetRegistryPriority());
+	FOnlineIdRegistryRegistry::Get().RegisterAccountIdRegistry(EOnlineServices::Epic, &FOnlineAccountIdRegistryEOSGS::Get(), GetRegistryPriority());
+	FOnlineIdRegistryRegistry::Get().RegisterSessionIdRegistry(EOnlineServices::Epic, &FOnlineSessionIdRegistryEOSGS::Get(), GetRegistryPriority());
+	FOnlineIdRegistryRegistry::Get().RegisterSessionInviteIdRegistry(EOnlineServices::Epic, &FOnlineSessionInviteIdRegistryEOSGS::Get(), GetRegistryPriority());
+
+	// Initialize the platform factory on startup.  This is necessary for the SDK to bind to rendering and input very early.
+	FOnlineServicesEOSGSPlatformFactory::Get();
+}
+
+void FOnlineServicesEOSGSModule::ShutdownModule()
+{
+	FOnlineServicesRegistry::Get().UnregisterServicesFactory(EOnlineServices::Epic, GetRegistryPriority());
+	FOnlineIdRegistryRegistry::Get().UnregisterAccountIdRegistry(EOnlineServices::Epic, GetRegistryPriority());
+	FOnlineIdRegistryRegistry::Get().UnregisterSessionIdRegistry(EOnlineServices::Epic, GetRegistryPriority());
+	FOnlineIdRegistryRegistry::Get().UnregisterSessionInviteIdRegistry(EOnlineServices::Epic, GetRegistryPriority());
+
+	UE::Online::FOnlineServicesEOSGSPlatformFactory::TearDown();
+}
+
+/* UE::Online */ }
+
+IMPLEMENT_MODULE(UE::Online::FOnlineServicesEOSGSModule, OnlineServicesEOSGS);
diff --git a/Engine/Plugins/Online/OnlineServicesEOSGS/Source/Private/OnlineServicesEOSGSModule.cpp b/Engine/Plugins/Online/OnlineServicesEOSGS/Source/Private/OnlineServicesEOSGSModule.cpp
deleted file mode 100644
diff --git a/Engine/Plugins/Online/OnlineServicesEOSGS/Source/Public/Online/OnlineServicesEOSGSModule.h b/Engine/Plugins/Online/OnlineServicesEOSGS/Source/Public/Online/OnlineServicesEOSGSModule.h
new file mode 100644
--- /dev/null
+++ b/Engine/Plugins/Online/OnlineServicesEOSGS/Source/Public/Online/OnlineServicesEOSGSModule.h
@@ -0,0 +1,21 @@
+// Copyright Epic Games, Inc. All Rights Reserved.
+
+#pragma once
+
+#include "Modules/ModuleInterface.h"
+
+namespace UE::Online {
+
+class FOnlineServicesEOSGSModule : public IModuleInterface
+{
+public:
+	ONLINESERVICESEOSGS_API static int GetRegistryPriority();
+
+private:
+	// Begin IModuleInterface
+	virtual void StartupModule() override;
+	virtual void ShutdownModule() override;
+	// End IModuleInterface
+};
+
+} // namespace UE::Online
\ No newline at end of file
