- Defer creation of the "Default" EOSGS platform instance in fork runs until it's requested

diff --git a/Engine/Plugins/Online/OnlineServicesEOSGS/Source/Private/Online/OnlineServicesEOSGSPlatformFactory.cpp b/Engine/Plugins/Online/OnlineServicesEOSGS/Source/Private/Online/OnlineServicesEOSGSPlatformFactory.cpp
--- a/Engine/Plugins/Online/OnlineServicesEOSGS/Source/Private/Online/OnlineServicesEOSGSPlatformFactory.cpp
+++ b/Engine/Plugins/Online/OnlineServicesEOSGS/Source/Private/Online/OnlineServicesEOSGSPlatformFactory.cpp
@@ -7,6 +7,7 @@
 #include "Online/OnlineServicesEOSGS.h"
 
 #include "Misc/ConfigCacheIni.h"
+#include "Misc/Fork.h"
 #include "Misc/LazySingleton.h"
 #include "Modules/ModuleManager.h"
 
@@ -52,7 +53,11 @@ namespace UE::Online {
 
 FOnlineServicesEOSGSPlatformFactory::FOnlineServicesEOSGSPlatformFactory()
 {
-	DefaultEOSPlatformHandle = CreatePlatform();
+	// If a fork is requested, we need to wait for post-fork to create the default platform
+	if (!FForkProcessHelper::IsForkRequested() || FForkProcessHelper::IsForkedChildProcess())
+	{
+		GetDefaultPlatform();
+	}
 }
 
 FOnlineServicesEOSGSPlatformFactory& FOnlineServicesEOSGSPlatformFactory::Get()
@@ -139,5 +144,14 @@ IEOSPlatformHandlePtr FOnlineServicesEOSGSPlatformFactory::CreatePlatform()
 	return EOSPlatformHandle;
 }
 
+IEOSPlatformHandlePtr FOnlineServicesEOSGSPlatformFactory::GetDefaultPlatform()
+{
+	if (!DefaultEOSPlatformHandle)
+	{
+		DefaultEOSPlatformHandle = CreatePlatform();
+	}
+
+	return DefaultEOSPlatformHandle;
+}
 
 /* UE::Online */ }
diff --git a/Engine/Plugins/Online/OnlineServicesEOSGS/Source/Private/Online/OnlineServicesEOSGSPlatformFactory.h b/Engine/Plugins/Online/OnlineServicesEOSGS/Source/Private/Online/OnlineServicesEOSGSPlatformFactory.h
--- a/Engine/Plugins/Online/OnlineServicesEOSGS/Source/Private/Online/OnlineServicesEOSGSPlatformFactory.h
+++ b/Engine/Plugins/Online/OnlineServicesEOSGS/Source/Private/Online/OnlineServicesEOSGSPlatformFactory.h
@@ -34,9 +34,8 @@ public:
 	 * The default platform is created on startup and can be used by the default instance of online services EOS.
 	 * @return the default platform
 	 */
-	IEOSPlatformHandlePtr GetDefaultPlatform() { return DefaultEOSPlatformHandle; }
+	IEOSPlatformHandlePtr GetDefaultPlatform();
 private:
-	/** Default constructor */
 	FOnlineServicesEOSGSPlatformFactory();
 	friend FLazySingleton;
 	/** Platform handle for the default instance */
