Replace usage of LogTemp with proper log categories

diff --git a/Engine/Plugins/Online/OnlineServices/Source/OnlineServicesCommon/Private/Online/AchievementsCommon.cpp b/Engine/Plugins/Online/OnlineServices/Source/OnlineServicesCommon/Private/Online/AchievementsCommon.cpp
--- a/Engine/Plugins/Online/OnlineServices/Source/OnlineServicesCommon/Private/Online/AchievementsCommon.cpp
+++ b/Engine/Plugins/Online/OnlineServices/Source/OnlineServicesCommon/Private/Online/AchievementsCommon.cpp
@@ -215,7 +215,7 @@ bool FAchievementsCommon::MeetUnlockCondition(const FAchievementUnlockRule& Achi
 			const FStatValue* StatValue = Stats.Find(Condition.StatName);
 			if (!StatValue)
 			{
-				UE_LOG(LogTemp, Warning, TEXT("Can't find stat %s when check if it can unlock achievement."), *Condition.StatName);
+				UE_LOG(LogOnlineServices, Warning, TEXT("Can't find stat %s when check if it can unlock achievement."), *Condition.StatName);
 				return false;
 			}
 
@@ -259,7 +259,7 @@ bool FAchievementsCommon::IsUnlocked(const FAccountId& AccountId, const FString&
 		return FMath::IsNearlyEqual(AchievementState.Progress, 1.0f);
 	}
 
-	UE_LOG(LogTemp, Warning, TEXT("Can't find state of achievement %s when check if it's unlocked."), *AchievementName);
+	UE_LOG(LogOnlineServices, Warning, TEXT("Can't find state of achievement %s when check if it's unlocked."), *AchievementName);
 
 	return false;
 }
diff --git a/Engine/Plugins/Online/OnlineServices/Source/OnlineServicesCommon/Private/Online/LobbiesCommonTests.cpp b/Engine/Plugins/Online/OnlineServices/Source/OnlineServicesCommon/Private/Online/LobbiesCommonTests.cpp
--- a/Engine/Plugins/Online/OnlineServices/Source/OnlineServicesCommon/Private/Online/LobbiesCommonTests.cpp
+++ b/Engine/Plugins/Online/OnlineServices/Source/OnlineServicesCommon/Private/Online/LobbiesCommonTests.cpp
@@ -1291,19 +1291,19 @@ TFunction<TFuture<void>(TOnlineAsyncOp<OpType>&, typename SecondaryOpType::Param
 		TPromise<void> Promise;
 		auto Future = Promise.GetFuture();
 
-		UE_LOG(LogTemp, Log, TEXT("[LobbiesFunctionalTest::ConsumeStepResult] Starting secondary operation %s."), SecondaryOpType::Name);
+		UE_LOG(LogOnlineServices, Log, TEXT("[LobbiesFunctionalTest::ConsumeStepResult] Starting secondary operation %s."), SecondaryOpType::Name);
 
 		Func(MoveTempIfPossible(InParams))
 		.Then([Promise = MoveTemp(Promise), Op = InAsyncOp.AsShared()](TFuture<TOnlineResult<SecondaryOpType>>&& Future) mutable
 		{
 			if (Future.Get().IsError())
 			{
-				UE_LOG(LogTemp, Log, TEXT("[LobbiesFunctionalTest::ConsumeStepResult] Failed secondary operation %s. Error: %s"), SecondaryOpType::Name, *Future.Get().GetErrorValue().GetLogString());
+				UE_LOG(LogOnlineServices, Log, TEXT("[LobbiesFunctionalTest::ConsumeStepResult] Failed secondary operation %s. Error: %s"), SecondaryOpType::Name, *Future.Get().GetErrorValue().GetLogString());
 				Op->SetError(Errors::RequestFailure(MoveTempIfPossible(Future.Get().GetErrorValue())));
 			}
 			else
 			{
-				UE_LOG(LogTemp, Log, TEXT("[LobbiesFunctionalTest::ConsumeStepResult] Completed secondary operation %s."), SecondaryOpType::Name);
+				UE_LOG(LogOnlineServices, Log, TEXT("[LobbiesFunctionalTest::ConsumeStepResult] Completed secondary operation %s."), SecondaryOpType::Name);
 			}
 			Promise.EmplaceValue();
 		});
@@ -1320,19 +1320,19 @@ TFunction<TFuture<void>(TOnlineAsyncOp<OpType>&)> ConsumeStepResult(TOnlineAsync
 		TPromise<void> Promise;
 		auto Future = Promise.GetFuture();
 
-		UE_LOG(LogTemp, Log, TEXT("[LobbiesFunctionalTest::ConsumeStepResult] Starting secondary operation %s."), SecondaryOpType::Name);
+		UE_LOG(LogOnlineServices, Log, TEXT("[LobbiesFunctionalTest::ConsumeStepResult] Starting secondary operation %s."), SecondaryOpType::Name);
 
 		Func(MoveTempIfPossible(InParams))
 		.Then([Promise = MoveTemp(Promise), Op = InAsyncOp.AsShared()](TFuture<TOnlineResult<SecondaryOpType>>&& Future) mutable
 		{
 			if (Future.Get().IsError())
 			{
-				UE_LOG(LogTemp, Log, TEXT("[LobbiesFunctionalTest::ConsumeStepResult] Failed secondary operation %s. Error: %s"), SecondaryOpType::Name, *Future.Get().GetErrorValue().GetLogString());
+				UE_LOG(LogOnlineServices, Log, TEXT("[LobbiesFunctionalTest::ConsumeStepResult] Failed secondary operation %s. Error: %s"), SecondaryOpType::Name, *Future.Get().GetErrorValue().GetLogString());
 				Op->SetError(Errors::RequestFailure(MoveTempIfPossible(Future.Get().GetErrorValue())));
 			}
 			else
 			{
-				UE_LOG(LogTemp, Log, TEXT("[LobbiesFunctionalTest::ConsumeStepResult] Completed secondary operation %s."), SecondaryOpType::Name);
+				UE_LOG(LogOnlineServices, Log, TEXT("[LobbiesFunctionalTest::ConsumeStepResult] Completed secondary operation %s."), SecondaryOpType::Name);
 			}
 			Promise.EmplaceValue();
 		});
@@ -1349,19 +1349,19 @@ TFunction<TFuture<void>(TOnlineAsyncOp<OpType>&, typename SecondaryOpType::Param
 		TPromise<void> Promise;
 		auto Future = Promise.GetFuture();
 
-		UE_LOG(LogTemp, Log, TEXT("[LobbiesFunctionalTest::CaptureStepResult] Starting secondary operation %s."), SecondaryOpType::Name);
+		UE_LOG(LogOnlineServices, Log, TEXT("[LobbiesFunctionalTest::CaptureStepResult] Starting secondary operation %s."), SecondaryOpType::Name);
 
 		Func(MoveTempIfPossible(InParams))
 		.Then([Promise = MoveTemp(Promise), Op = InAsyncOp.AsShared(), ResultKey](TFuture<TOnlineResult<SecondaryOpType>>&& Future) mutable
 		{
 			if (Future.Get().IsError())
 			{
-				UE_LOG(LogTemp, Log, TEXT("[LobbiesFunctionalTest::CaptureStepResult] Failed secondary operation %s. Error: %s"), SecondaryOpType::Name, *Future.Get().GetErrorValue().GetLogString());
+				UE_LOG(LogOnlineServices, Log, TEXT("[LobbiesFunctionalTest::CaptureStepResult] Failed secondary operation %s. Error: %s"), SecondaryOpType::Name, *Future.Get().GetErrorValue().GetLogString());
 				Op->SetError(Errors::RequestFailure(MoveTempIfPossible(Future.Get().GetErrorValue())));
 			}
 			else
 			{
-				UE_LOG(LogTemp, Log, TEXT("[LobbiesFunctionalTest::CaptureStepResult] Captured secondary operation %s as %s"), SecondaryOpType::Name, *ResultKey);
+				UE_LOG(LogOnlineServices, Log, TEXT("[LobbiesFunctionalTest::CaptureStepResult] Captured secondary operation %s as %s"), SecondaryOpType::Name, *ResultKey);
 				Op->Data.Set(ResultKey, MoveTempIfPossible(Future.Get().GetOkValue()));
 			}
 			Promise.EmplaceValue();
@@ -1379,19 +1379,19 @@ TFunction<TFuture<void>(TOnlineAsyncOp<OpType>&)> CaptureStepResult(TOnlineAsync
 		TPromise<void> Promise;
 		auto Future = Promise.GetFuture();
 
-		UE_LOG(LogTemp, Log, TEXT("[LobbiesFunctionalTest::CaptureStepResult] Starting secondary operation %s."), SecondaryOpType::Name);
+		UE_LOG(LogOnlineServices, Log, TEXT("[LobbiesFunctionalTest::CaptureStepResult] Starting secondary operation %s."), SecondaryOpType::Name);
 
 		Func(MoveTempIfPossible(InParams))
 		.Then([Promise = MoveTemp(Promise), Op = InAsyncOp.AsShared(), ResultKey](TFuture<TOnlineResult<SecondaryOpType>>&& Future) mutable
 		{
 			if (Future.Get().IsError())
 			{
-				UE_LOG(LogTemp, Log, TEXT("[LobbiesFunctionalTest::CaptureStepResult] Failed secondary operation %s. Error: %s"), SecondaryOpType::Name, *Future.Get().GetErrorValue().GetLogString());
+				UE_LOG(LogOnlineServices, Log, TEXT("[LobbiesFunctionalTest::CaptureStepResult] Failed secondary operation %s. Error: %s"), SecondaryOpType::Name, *Future.Get().GetErrorValue().GetLogString());
 				Op->SetError(Errors::RequestFailure(MoveTempIfPossible(Future.Get().GetErrorValue())));
 			}
 			else
 			{
-				UE_LOG(LogTemp, Log, TEXT("[LobbiesFunctionalTest::CaptureStepResult] Captured secondary operation %s as %s"), SecondaryOpType::Name, *ResultKey);
+				UE_LOG(LogOnlineServices, Log, TEXT("[LobbiesFunctionalTest::CaptureStepResult] Captured secondary operation %s as %s"), SecondaryOpType::Name, *ResultKey);
 				Op->Data.Set(ResultKey, MoveTempIfPossible(Future.Get().GetOkValue()));
 			}
 			Promise.EmplaceValue();
@@ -1409,19 +1409,19 @@ TFunction<TFuture<void>(TOnlineAsyncOp<OpType>&, typename SecondaryOpType::Param
 		auto Promise = MakeShared<TPromise<void>>();
 		auto Future = Promise->GetFuture();
 
-		UE_LOG(LogTemp, Log, TEXT("[LobbiesFunctionalTest::ConsumeOperationStepResult] Starting secondary operation %s."), SecondaryOpType::Name);
+		UE_LOG(LogOnlineServices, Log, TEXT("[LobbiesFunctionalTest::ConsumeOperationStepResult] Starting secondary operation %s."), SecondaryOpType::Name);
 
 		Func(MoveTempIfPossible(InParams))
 		.OnComplete([Promise, Op = InAsyncOp.AsShared()](const TOnlineResult<SecondaryOpType>& Result) mutable -> void
 		{
 			if (Result.IsError())
 			{
-				UE_LOG(LogTemp, Log, TEXT("[LobbiesFunctionalTest::ConsumeOperationStepResult] Failed secondary operation %s. Error: %s"), SecondaryOpType::Name, *Result.GetErrorValue().GetLogString());
+				UE_LOG(LogOnlineServices, Log, TEXT("[LobbiesFunctionalTest::ConsumeOperationStepResult] Failed secondary operation %s. Error: %s"), SecondaryOpType::Name, *Result.GetErrorValue().GetLogString());
 				Op->SetError(Errors::RequestFailure(MoveTempIfPossible(Result.GetErrorValue())));
 			}
 			else
 			{
-				UE_LOG(LogTemp, Log, TEXT("[LobbiesFunctionalTest::ConsumeOperationStepResult] Completed secondary operation %s."), SecondaryOpType::Name);
+				UE_LOG(LogOnlineServices, Log, TEXT("[LobbiesFunctionalTest::ConsumeOperationStepResult] Completed secondary operation %s."), SecondaryOpType::Name);
 			}
 			Promise->EmplaceValue();
 		});
@@ -1438,19 +1438,19 @@ TFunction<TFuture<void>(TOnlineAsyncOp<OpType>&)> ConsumeOperationStepResult(TOn
 		auto Promise = MakeShared<TPromise<void>>();
 		auto Future = Promise->GetFuture();
 
-		UE_LOG(LogTemp, Log, TEXT("[LobbiesFunctionalTest::ConsumeOperationStepResult] Starting secondary operation %s."), SecondaryOpType::Name);
+		UE_LOG(LogOnlineServices, Log, TEXT("[LobbiesFunctionalTest::ConsumeOperationStepResult] Starting secondary operation %s."), SecondaryOpType::Name);
 
 		Func(MoveTempIfPossible(InParams))
 		.OnComplete([Promise, Op = InAsyncOp.AsShared()](const TOnlineResult<SecondaryOpType>& Result) mutable
 		{
 			if (Result.IsError())
 			{
-				UE_LOG(LogTemp, Log, TEXT("[LobbiesFunctionalTest::ConsumeOperationStepResult] Failed secondary operation %s. Error: %s"), SecondaryOpType::Name, *Result.GetErrorValue().GetLogString());
+				UE_LOG(LogOnlineServices, Log, TEXT("[LobbiesFunctionalTest::ConsumeOperationStepResult] Failed secondary operation %s. Error: %s"), SecondaryOpType::Name, *Result.GetErrorValue().GetLogString());
 				Op->SetError(Errors::RequestFailure(MoveTempIfPossible(Result.GetErrorValue())));
 			}
 			else
 			{
-				UE_LOG(LogTemp, Log, TEXT("[LobbiesFunctionalTest::ConsumeOperationStepResult] Completed secondary operation %s."), SecondaryOpType::Name);
+				UE_LOG(LogOnlineServices, Log, TEXT("[LobbiesFunctionalTest::ConsumeOperationStepResult] Completed secondary operation %s."), SecondaryOpType::Name);
 			}
 			Promise->EmplaceValue();
 		});
@@ -1467,19 +1467,19 @@ TFunction<TFuture<void>(TOnlineAsyncOp<OpType>&, typename SecondaryOpType::Param
 		auto Promise = MakeShared<TPromise<void>>();
 		auto Future = Promise->GetFuture();
 
-		UE_LOG(LogTemp, Log, TEXT("[LobbiesFunctionalTest::CaptureOperationStepResult] Starting secondary operation %s."), SecondaryOpType::Name);
+		UE_LOG(LogOnlineServices, Log, TEXT("[LobbiesFunctionalTest::CaptureOperationStepResult] Starting secondary operation %s."), SecondaryOpType::Name);
 
 		Func(MoveTempIfPossible(InParams))
 		.OnComplete([Promise, Op = InAsyncOp.AsShared(), ResultKey](const TOnlineResult<SecondaryOpType>& Result) mutable -> void
 		{
 			if (Result.IsError())
 			{
-				UE_LOG(LogTemp, Log, TEXT("[LobbiesFunctionalTest::CaptureOperationStepResult] Failed secondary operation %s. Error: %s"), SecondaryOpType::Name, *Result.GetErrorValue().GetLogString());
+				UE_LOG(LogOnlineServices, Log, TEXT("[LobbiesFunctionalTest::CaptureOperationStepResult] Failed secondary operation %s. Error: %s"), SecondaryOpType::Name, *Result.GetErrorValue().GetLogString());
 				Op->SetError(Errors::RequestFailure(MoveTempIfPossible(Result.GetErrorValue())));
 			}
 			else
 			{
-				UE_LOG(LogTemp, Log, TEXT("[LobbiesFunctionalTest::CaptureOperationStepResult] Captured secondary operation %s as %s"), SecondaryOpType::Name, *ResultKey);
+				UE_LOG(LogOnlineServices, Log, TEXT("[LobbiesFunctionalTest::CaptureOperationStepResult] Captured secondary operation %s as %s"), SecondaryOpType::Name, *ResultKey);
 				Op->Data.Set(ResultKey, Result.GetOkValue());
 			}
 			Promise->EmplaceValue();
@@ -1497,19 +1497,19 @@ TFunction<TFuture<void>(TOnlineAsyncOp<OpType>&)> CaptureOperationStepResult(TOn
 		auto Promise = MakeShared<TPromise<void>>();
 		auto Future = Promise->GetFuture();
 
-		UE_LOG(LogTemp, Log, TEXT("[LobbiesFunctionalTest::CaptureOperationStepResult] Starting secondary operation %s."), SecondaryOpType::Name);
+		UE_LOG(LogOnlineServices, Log, TEXT("[LobbiesFunctionalTest::CaptureOperationStepResult] Starting secondary operation %s."), SecondaryOpType::Name);
 
 		Func(MoveTempIfPossible(InParams))
 		.OnComplete([Promise, Op = InAsyncOp.AsShared(), ResultKey](const TOnlineResult<SecondaryOpType>& Result) mutable
 		{
 			if (Result.IsError())
 			{
-				UE_LOG(LogTemp, Log, TEXT("[LobbiesFunctionalTest::CaptureOperationStepResult] Failed secondary operation %s. Error: %s"), SecondaryOpType::Name, *Result.GetErrorValue().GetLogString());
+				UE_LOG(LogOnlineServices, Log, TEXT("[LobbiesFunctionalTest::CaptureOperationStepResult] Failed secondary operation %s. Error: %s"), SecondaryOpType::Name, *Result.GetErrorValue().GetLogString());
 				Op->SetError(Errors::RequestFailure(MoveTempIfPossible(Result.GetErrorValue())));
 			}
 			else
 			{
-				UE_LOG(LogTemp, Log, TEXT("[LobbiesFunctionalTest::CaptureOperationStepResult] Captured secondary operation %s as %s"), SecondaryOpType::Name, *ResultKey);
+				UE_LOG(LogOnlineServices, Log, TEXT("[LobbiesFunctionalTest::CaptureOperationStepResult] Captured secondary operation %s as %s"), SecondaryOpType::Name, *ResultKey);
 				Op->Data.Set(ResultKey, Result.GetOkValue());
 			}
 			Promise->EmplaceValue();
diff --git a/Engine/Plugins/Online/OnlineServices/Source/OnlineServicesCommon/Private/Online/LobbiesCommonTypes.cpp b/Engine/Plugins/Online/OnlineServices/Source/OnlineServicesCommon/Private/Online/LobbiesCommonTypes.cpp
--- a/Engine/Plugins/Online/OnlineServices/Source/OnlineServicesCommon/Private/Online/LobbiesCommonTypes.cpp
+++ b/Engine/Plugins/Online/OnlineServices/Source/OnlineServicesCommon/Private/Online/LobbiesCommonTypes.cpp
@@ -2,6 +2,8 @@
 
 #include "Online/LobbiesCommonTypes.h"
 
+#include "Online/OnlineServicesLog.h"
+
 namespace UE::Online {
 
 FSchemaId LobbyBaseSchemaId = TEXT("LobbyBase");
@@ -33,7 +35,7 @@ TOnlineResult<FLobbyClientDataPrepareServiceSnapshot> FLobbyClientData::PrepareS
 		InternalPublicData->SchemaData.PrepareServiceSnapshot({ TOptional<FSchemaId>(), MoveTemp(Params.LobbySnapshot.SchemaServiceSnapshot) });
 	if (PrepareLobbyServiceSnapshotResult.IsError())
 	{
-		UE_LOG(LogTemp, Warning, TEXT("[FLobbyClientData::PrepareServiceSnapshot] Failed to prepare lobby service snapshot: Lobby[%s], Result[%s]"),
+		UE_LOG(LogOnlineServices, Warning, TEXT("[FLobbyClientData::PrepareServiceSnapshot] Failed to prepare lobby service snapshot: Lobby[%s], Result[%s]"),
 			*ToLogString(InternalPublicData->LobbyId), *PrepareLobbyServiceSnapshotResult.GetErrorValue().GetLogString());
 		return TOnlineResult<FLobbyClientDataPrepareServiceSnapshot>(MoveTemp(PrepareLobbyServiceSnapshotResult.GetErrorValue()));
 	}
@@ -58,7 +60,7 @@ TOnlineResult<FLobbyClientDataPrepareServiceSnapshot> FLobbyClientData::PrepareS
 			{
 				const ELobbyMemberLeaveReason DefaultLeaveReason = ELobbyMemberLeaveReason::Disconnected;
 				NewPreparedServiceChanges.LeavingMembers.Add({ MemberData.Value, DefaultLeaveReason });
-				UE_LOG(LogTemp, Warning, TEXT("[FLobbyClientData::PrepareServiceSnapshot] Member left lobby without giving reason, defaulting to %s: Lobby[%s], Member[%s]"),
+				UE_LOG(LogOnlineServices, Warning, TEXT("[FLobbyClientData::PrepareServiceSnapshot] Member left lobby without giving reason, defaulting to %s: Lobby[%s], Member[%s]"),
 					*ToLogString(InternalPublicData->LobbyId), LexToString(DefaultLeaveReason), *ToLogString(MemberData.Key));
 			}
 		}
@@ -70,7 +72,7 @@ TOnlineResult<FLobbyClientDataPrepareServiceSnapshot> FLobbyClientData::PrepareS
 	{
 		if (!Params.LobbySnapshot.Members.Contains(MemberSnapshotPair.Key))
 		{
-			UE_LOG(LogTemp, Warning, TEXT("[FLobbyClientData::PrepareServiceSnapshot] Member update ignored for unknown member: Lobby[%s], Member[%s]"),
+			UE_LOG(LogOnlineServices, Warning, TEXT("[FLobbyClientData::PrepareServiceSnapshot] Member update ignored for unknown member: Lobby[%s], Member[%s]"),
 				*ToLogString(InternalPublicData->LobbyId), *ToLogString(MemberSnapshotPair.Key));
 		}
 		else
@@ -98,7 +100,7 @@ TOnlineResult<FLobbyClientDataPrepareServiceSnapshot> FLobbyClientData::PrepareS
 				PreparingMemberData->SchemaData.PrepareServiceSnapshot({ NewDerivedSchemaId, MoveTemp(MemberSnapshot.SchemaServiceSnapshot) });
 			if (PrepareLobbyMemberServiceSnapshotResult.IsError())
 			{
-				UE_LOG(LogTemp, Warning, TEXT("[FLobbyClientData::PrepareServiceSnapshot] Failed to prepare lobby member service snapshot: Lobby[%s], Member[%s], Result[%s]"),
+				UE_LOG(LogOnlineServices, Warning, TEXT("[FLobbyClientData::PrepareServiceSnapshot] Failed to prepare lobby member service snapshot: Lobby[%s], Member[%s], Result[%s]"),
 					*ToLogString(InternalPublicData->LobbyId), *ToLogString(MemberSnapshot.AccountId), *PrepareLobbyMemberServiceSnapshotResult.GetErrorValue().GetLogString());
 				return TOnlineResult<FLobbyClientDataPrepareServiceSnapshot>(MoveTemp(PrepareLobbyMemberServiceSnapshotResult.GetErrorValue()));
 			}
@@ -110,7 +112,7 @@ TOnlineResult<FLobbyClientDataPrepareServiceSnapshot> FLobbyClientData::PrepareS
 	{
 		if (!Params.LobbySnapshot.Members.Contains(Params.LobbySnapshot.OwnerAccountId))
 		{
-			UE_LOG(LogTemp, Warning, TEXT("[FLobbyClientData::PrepareServiceSnapshot] Lobby owner chage failed - lobby member data not found: Lobby[%s], Member[%s]"),
+			UE_LOG(LogOnlineServices, Warning, TEXT("[FLobbyClientData::PrepareServiceSnapshot] Lobby owner chage failed - lobby member data not found: Lobby[%s], Member[%s]"),
 				*ToLogString(InternalPublicData->LobbyId), *ToLogString(Params.LobbySnapshot.OwnerAccountId));
 			return TOnlineResult<FLobbyClientDataPrepareServiceSnapshot>(Errors::InvalidState());
 		}
@@ -143,7 +145,7 @@ FLobbyClientDataCommitServiceSnapshot::Result FLobbyClientData::CommitServiceSna
 
 	if (!PreparedServiceChanges)
 	{
-		UE_LOG(LogTemp, Warning, TEXT("[FLobbyClientData::CommitServiceSnapshot] Commit failed, no pending changes found: Lobby[%s]"),
+		UE_LOG(LogOnlineServices, Warning, TEXT("[FLobbyClientData::CommitServiceSnapshot] Commit failed, no pending changes found: Lobby[%s]"),
 			*ToLogString(InternalPublicData->LobbyId));
 		return CommitResult;
 	}
@@ -301,7 +303,7 @@ TOnlineResult<FLobbyClientDataPrepareClientChanges> FLobbyClientData::PrepareCli
 		{
 			if (InternalPublicData->LocalName != FName())
 			{
-				UE_LOG(LogTemp, Warning, TEXT("[FLobbyClientData::PrepareClientChanges] Local name chage failed - local name cannot be changed once set: Lobby[%s], CurrentName[%s], NewName[%s]"),
+				UE_LOG(LogOnlineServices, Warning, TEXT("[FLobbyClientData::PrepareClientChanges] Local name chage failed - local name cannot be changed once set: Lobby[%s], CurrentName[%s], NewName[%s]"),
 					*ToLogString(InternalPublicData->LobbyId), *InternalPublicData->LocalName.ToString().ToLower(), *Params.ClientChanges.LocalName->ToString().ToLower());
 				return TOnlineResult<FLobbyClientDataPrepareClientChanges>(Errors::InvalidParams());
 			}
@@ -322,7 +324,7 @@ TOnlineResult<FLobbyClientDataPrepareClientChanges> FLobbyClientData::PrepareCli
 		{
 			if (!MemberDataStorage.Contains(*Params.ClientChanges.OwnerAccountId))
 			{
-				UE_LOG(LogTemp, Warning, TEXT("[FLobbyClientData::PrepareClientChanges] Lobby owner chage failed - new owner is not in lobby: Lobby[%s], NewOwner[%s]"),
+				UE_LOG(LogOnlineServices, Warning, TEXT("[FLobbyClientData::PrepareClientChanges] Lobby owner chage failed - new owner is not in lobby: Lobby[%s], NewOwner[%s]"),
 					*ToLogString(InternalPublicData->LobbyId), *ToLogString(*Params.ClientChanges.OwnerAccountId));
 				return TOnlineResult<FLobbyClientDataPrepareClientChanges>(Errors::InvalidParams());
 			}
@@ -337,7 +339,7 @@ TOnlineResult<FLobbyClientDataPrepareClientChanges> FLobbyClientData::PrepareCli
 		{
 			if (!MemberDataStorage.Contains(*Params.ClientChanges.KickedTargetMember))
 			{
-				UE_LOG(LogTemp, Warning, TEXT("[FLobbyClientData::PrepareClientChanges] Kick target member failed - target is not in the lobby: Lobby[%s], KickedTargetMember[%s]"),
+				UE_LOG(LogOnlineServices, Warning, TEXT("[FLobbyClientData::PrepareClientChanges] Kick target member failed - target is not in the lobby: Lobby[%s], KickedTargetMember[%s]"),
 					*ToLogString(InternalPublicData->LobbyId), *ToLogString(*Params.ClientChanges.KickedTargetMember));
 				return TOnlineResult<FLobbyClientDataPrepareClientChanges>(Errors::InvalidParams());
 			}
@@ -371,7 +373,7 @@ TOnlineResult<FLobbyClientDataPrepareClientChanges> FLobbyClientData::PrepareCli
 				InternalPublicData->SchemaData.PrepareClientChanges(MoveTemp(PrepareClientChangesParams));
 			if (PrepareClientChangesResult.IsError())
 			{
-				UE_LOG(LogTemp, Warning, TEXT("[FLobbyClientData::PrepareClientChanges] Failed to prepare lobby attribute changes: Lobby[%s], Result[%s]"),
+				UE_LOG(LogOnlineServices, Warning, TEXT("[FLobbyClientData::PrepareClientChanges] Failed to prepare lobby attribute changes: Lobby[%s], Result[%s]"),
 					*ToLogString(InternalPublicData->LobbyId), *PrepareClientChangesResult.GetErrorValue().GetLogString());
 				return TOnlineResult<FLobbyClientDataPrepareClientChanges>(MoveTemp(PrepareClientChangesResult.GetErrorValue()));
 			}
@@ -388,7 +390,7 @@ TOnlineResult<FLobbyClientDataPrepareClientChanges> FLobbyClientData::PrepareCli
 			// Check that the local user can make lobby changes.
 			if (Params.LocalAccountId != InternalPublicData->OwnerAccountId)
 			{
-				UE_LOG(LogTemp, Warning, TEXT("[FLobbyClientData::PrepareClientChanges] Failed to prepare lobby changes. LocalAccountId is not the lobby owner: Lobby[%s], LocalAccountId[%s], OwnerAccountId[%s]"),
+				UE_LOG(LogOnlineServices, Warning, TEXT("[FLobbyClientData::PrepareClientChanges] Failed to prepare lobby changes. LocalAccountId is not the lobby owner: Lobby[%s], LocalAccountId[%s], OwnerAccountId[%s]"),
 					*ToLogString(InternalPublicData->LobbyId), *ToLogString(Params.LocalAccountId), *ToLogString(InternalPublicData->OwnerAccountId));
 				return TOnlineResult<FLobbyClientDataPrepareClientChanges>(Errors::InvalidParams());
 			}
@@ -396,7 +398,7 @@ TOnlineResult<FLobbyClientDataPrepareClientChanges> FLobbyClientData::PrepareCli
 			// Check that the local user is not trying to change the lobby while leaving.
 			if (Params.ClientChanges.LocalUserLeaveReason)
 			{
-				UE_LOG(LogTemp, Warning, TEXT("[FLobbyClientData::PrepareClientChanges] Failed to prepare lobby changes. Unable to modify the lobby while leaving: Lobby[%s], LocalAccountId[%s]"),
+				UE_LOG(LogOnlineServices, Warning, TEXT("[FLobbyClientData::PrepareClientChanges] Failed to prepare lobby changes. Unable to modify the lobby while leaving: Lobby[%s], LocalAccountId[%s]"),
 					*ToLogString(InternalPublicData->LobbyId), *ToLogString(Params.LocalAccountId));
 				return TOnlineResult<FLobbyClientDataPrepareClientChanges>(Errors::InvalidParams());
 			}
@@ -408,7 +410,7 @@ TOnlineResult<FLobbyClientDataPrepareClientChanges> FLobbyClientData::PrepareCli
 		// Check that the local user is not joining and leaving.
 		if (Params.ClientChanges.MemberAttributes && Params.ClientChanges.LocalUserLeaveReason)
 		{
-			UE_LOG(LogTemp, Warning, TEXT("[FLobbyClientData::PrepareClientChanges] Failed to prepare lobby member changes. Member cannot be updated and removed: Lobby[%s], LocalAccountId[%s]"),
+			UE_LOG(LogOnlineServices, Warning, TEXT("[FLobbyClientData::PrepareClientChanges] Failed to prepare lobby member changes. Member cannot be updated and removed: Lobby[%s], LocalAccountId[%s]"),
 				*ToLogString(InternalPublicData->LobbyId), *ToLogString(Params.LocalAccountId));
 			return TOnlineResult<FLobbyClientDataPrepareClientChanges>(Errors::InvalidParams());
 		}
@@ -453,7 +455,7 @@ TOnlineResult<FLobbyClientDataPrepareClientChanges> FLobbyClientData::PrepareCli
 				PreparingMemberData->SchemaData.PrepareClientChanges(MoveTemp(PrepareClientChangesParams));
 			if (PrepareClientChangesResult.IsError())
 			{
-				UE_LOG(LogTemp, Warning, TEXT("[FLobbyClientData::PrepareClientChanges] Failed to prepare lobby member attribute changes: Lobby[%s], LocalMember[%s], Result[%s]"),
+				UE_LOG(LogOnlineServices, Warning, TEXT("[FLobbyClientData::PrepareClientChanges] Failed to prepare lobby member attribute changes: Lobby[%s], LocalMember[%s], Result[%s]"),
 					*ToLogString(InternalPublicData->LobbyId), *ToLogString(Params.LocalAccountId), *PrepareClientChangesResult.GetErrorValue().GetLogString());
 				return TOnlineResult<FLobbyClientDataPrepareClientChanges>(MoveTemp(PrepareClientChangesResult.GetErrorValue()));
 			}
@@ -468,7 +470,7 @@ TOnlineResult<FLobbyClientDataPrepareClientChanges> FLobbyClientData::PrepareCli
 		{
 			if (!MemberDataStorage.Contains(Params.LocalAccountId))
 			{
-				UE_LOG(LogTemp, Warning, TEXT("[FLobbyClientData::PrepareClientChanges] Local user leave faild - user is not in the lobby: Lobby[%s], LocalAccountId[%s]"),
+				UE_LOG(LogOnlineServices, Warning, TEXT("[FLobbyClientData::PrepareClientChanges] Local user leave faild - user is not in the lobby: Lobby[%s], LocalAccountId[%s]"),
 					*ToLogString(InternalPublicData->LobbyId), *ToLogString(Params.LocalAccountId));
 				return TOnlineResult<FLobbyClientDataPrepareClientChanges>(Errors::InvalidParams());
 			}
@@ -488,7 +490,7 @@ FLobbyClientDataCommitClientChanges::Result FLobbyClientData::CommitClientChange
 
 	if (!PreparedClientChanges)
 	{
-		UE_LOG(LogTemp, Warning, TEXT("[FLobbyClientData::CommitClientChanges] Commit failed, no pending changes found: Lobby[%s]"),
+		UE_LOG(LogOnlineServices, Warning, TEXT("[FLobbyClientData::CommitClientChanges] Commit failed, no pending changes found: Lobby[%s]"),
 			*ToLogString(InternalPublicData->LobbyId));
 		return CommitResult;
 	}
diff --git a/Engine/Plugins/Online/OnlineServices/Source/OnlineServicesCommon/Private/Online/OnlineServicesCommon.cpp b/Engine/Plugins/Online/OnlineServices/Source/OnlineServicesCommon/Private/Online/OnlineServicesCommon.cpp
--- a/Engine/Plugins/Online/OnlineServices/Source/OnlineServicesCommon/Private/Online/OnlineServicesCommon.cpp
+++ b/Engine/Plugins/Online/OnlineServices/Source/OnlineServicesCommon/Private/Online/OnlineServicesCommon.cpp
@@ -6,8 +6,6 @@
 #include "Online/OnlineAsyncOpCache_Meta.inl"
 #include "Online/OnlineExecHandler.h"
 
-DEFINE_LOG_CATEGORY(LogOnlineServices);
-
 namespace UE::Online {
 
 struct FOnlineServicesCommonConfig
diff --git a/Engine/Plugins/Online/OnlineServices/Source/OnlineServicesCommon/Private/Online/SessionsCommon.cpp b/Engine/Plugins/Online/OnlineServices/Source/OnlineServicesCommon/Private/Online/SessionsCommon.cpp
--- a/Engine/Plugins/Online/OnlineServices/Source/OnlineServicesCommon/Private/Online/SessionsCommon.cpp
+++ b/Engine/Plugins/Online/OnlineServices/Source/OnlineServicesCommon/Private/Online/SessionsCommon.cpp
@@ -38,7 +38,7 @@ namespace UE::Online {
 
 	void FSessionCommon::DumpState() const
 	{
-		UE_LOG(LogTemp, Log, TEXT("[ISession]"));
+		UE_LOG(LogOnlineServices, Log, TEXT("[ISession]"));
 		DumpMemberData();
 		DumpSessionInfo();
 		DumpSessionSettings();
@@ -46,41 +46,41 @@ namespace UE::Online {
 
 	void FSessionCommon::DumpMemberData() const
 	{
-		UE_LOG(LogTemp, Log, TEXT("OwnerAccountId [%s]"), *UE::Online::ToLogString(OwnerAccountId));
-		UE_LOG(LogTemp, Log, TEXT("SessionMemberIdsSet:"));
+		UE_LOG(LogOnlineServices, Log, TEXT("OwnerAccountId [%s]"), *UE::Online::ToLogString(OwnerAccountId));
+		UE_LOG(LogOnlineServices, Log, TEXT("SessionMemberIdsSet:"));
 
 		for (const FAccountId& SessionMemberId : SessionMembers)
 		{
-			UE_LOG(LogTemp, Log, TEXT("	[%s]"), *UE::Online::ToLogString(SessionMemberId));
+			UE_LOG(LogOnlineServices, Log, TEXT("	[%s]"), *UE::Online::ToLogString(SessionMemberId));
 		}
 	}
 
 	void FSessionCommon::DumpSessionInfo() const
 	{
-		UE_LOG(LogTemp, Log, TEXT("SessionInfo:"));
-		UE_LOG(LogTemp, Log, TEXT("	SessionId [%s]"), *UE::Online::ToLogString(SessionInfo.SessionId));
-		UE_LOG(LogTemp, Log, TEXT("	bAllowSanctionedPlayers [%d]"), SessionInfo.bAllowSanctionedPlayers);
-		UE_LOG(LogTemp, Log, TEXT("	bAntiCheatProtected [%d]"), SessionInfo.bAntiCheatProtected);
-		UE_LOG(LogTemp, Log, TEXT("	bIsDedicatedServerSession [%d]"), SessionInfo.bIsDedicatedServerSession);
-		UE_LOG(LogTemp, Log, TEXT("	bIsLANSession [%d]"), SessionInfo.bIsLANSession);
-		UE_LOG(LogTemp, Log, TEXT("	SessionIdOverride [%s]"), *SessionInfo.SessionIdOverride);
+		UE_LOG(LogOnlineServices, Log, TEXT("SessionInfo:"));
+		UE_LOG(LogOnlineServices, Log, TEXT("	SessionId [%s]"), *UE::Online::ToLogString(SessionInfo.SessionId));
+		UE_LOG(LogOnlineServices, Log, TEXT("	bAllowSanctionedPlayers [%d]"), SessionInfo.bAllowSanctionedPlayers);
+		UE_LOG(LogOnlineServices, Log, TEXT("	bAntiCheatProtected [%d]"), SessionInfo.bAntiCheatProtected);
+		UE_LOG(LogOnlineServices, Log, TEXT("	bIsDedicatedServerSession [%d]"), SessionInfo.bIsDedicatedServerSession);
+		UE_LOG(LogOnlineServices, Log, TEXT("	bIsLANSession [%d]"), SessionInfo.bIsLANSession);
+		UE_LOG(LogOnlineServices, Log, TEXT("	SessionIdOverride [%s]"), *SessionInfo.SessionIdOverride);
 	}
 
 	void FSessionCommon::DumpSessionSettings() const
 	{
-		UE_LOG(LogTemp, Log, TEXT("SessionSettings:"));
-		UE_LOG(LogTemp, Log, TEXT("	bAllowNewMembers [%d]"), SessionSettings.bAllowNewMembers);
-		UE_LOG(LogTemp, Log, TEXT("	JoinPolicy [%s]"), LexToString(SessionSettings.JoinPolicy));
-		UE_LOG(LogTemp, Log, TEXT("	NumMaxConnections [%d]"), SessionSettings.NumMaxConnections);
-		UE_LOG(LogTemp, Log, TEXT("	SchemaName [%s]"), *SessionSettings.SchemaName.ToString());
-		UE_LOG(LogTemp, Log, TEXT("	CustomSettings:"));
+		UE_LOG(LogOnlineServices, Log, TEXT("SessionSettings:"));
+		UE_LOG(LogOnlineServices, Log, TEXT("	bAllowNewMembers [%d]"), SessionSettings.bAllowNewMembers);
+		UE_LOG(LogOnlineServices, Log, TEXT("	JoinPolicy [%s]"), LexToString(SessionSettings.JoinPolicy));
+		UE_LOG(LogOnlineServices, Log, TEXT("	NumMaxConnections [%d]"), SessionSettings.NumMaxConnections);
+		UE_LOG(LogOnlineServices, Log, TEXT("	SchemaName [%s]"), *SessionSettings.SchemaName.ToString());
+		UE_LOG(LogOnlineServices, Log, TEXT("	CustomSettings:"));
 
 		for (const TPair<FSchemaAttributeId, FCustomSessionSetting>& Entry : SessionSettings.CustomSettings)
 		{
-			UE_LOG(LogTemp, Log, TEXT("		Key: [%s]"), *Entry.Key.ToString());
-			UE_LOG(LogTemp, Log, TEXT("		Value:	Visibility [%s]"), LexToString(Entry.Value.Visibility));
-			UE_LOG(LogTemp, Log, TEXT("				ID [%d]"), Entry.Value.ID);
-			UE_LOG(LogTemp, Log, TEXT("				Data [%s]"), *LexToString(Entry.Value.Data));
+			UE_LOG(LogOnlineServices, Log, TEXT("		Key: [%s]"), *Entry.Key.ToString());
+			UE_LOG(LogOnlineServices, Log, TEXT("		Value:	Visibility [%s]"), LexToString(Entry.Value.Visibility));
+			UE_LOG(LogOnlineServices, Log, TEXT("				ID [%d]"), Entry.Value.ID);
+			UE_LOG(LogOnlineServices, Log, TEXT("				Data [%s]"), *LexToString(Entry.Value.Data));
 		}
 	}
 
@@ -143,14 +143,14 @@ namespace UE::Online {
 #define CHECK_PARAMS_ID_HANDLE(IdHandle, MethodName) \
 	if (!IdHandle.IsValid()) \
 	{ \
-		UE_LOG(LogTemp, Verbose, TEXT("[%s] Value of [%s] invalid."), UTF8_TO_TCHAR(__FUNCTION__), UTF8_TO_TCHAR(#IdHandle)); \
+		UE_LOG(LogOnlineServices, Verbose, TEXT("[%s] Value of [%s] invalid."), UTF8_TO_TCHAR(__FUNCTION__), UTF8_TO_TCHAR(#IdHandle)); \
 		return TOnlineResult<MethodName>(Errors::InvalidParams()); \
 	} \
 
 #define CHECK_PARAMS_FNAME(Name, MethodName) \
 	if (Name.IsNone()) \
 	{ \
-		UE_LOG(LogTemp, Verbose, TEXT("[%s] Value of [%s] invalid."), UTF8_TO_TCHAR(__FUNCTION__), UTF8_TO_TCHAR(#Name)); \
+		UE_LOG(LogOnlineServices, Verbose, TEXT("[%s] Value of [%s] invalid."), UTF8_TO_TCHAR(__FUNCTION__), UTF8_TO_TCHAR(#Name)); \
 		return TOnlineResult<MethodName>(Errors::InvalidParams()); \
 	} \
 
@@ -210,7 +210,7 @@ namespace UE::Online {
 		}
 		else
 		{
-			UE_LOG(LogTemp, Verbose, TEXT("[%s] No data found for user [%s]."), UTF8_TO_TCHAR(__FUNCTION__), *ToLogString(Params.LocalAccountId));
+			UE_LOG(LogOnlineServices, Verbose, TEXT("[%s] No data found for user [%s]."), UTF8_TO_TCHAR(__FUNCTION__), *ToLogString(Params.LocalAccountId));
 			return TOnlineResult<FGetPresenceSession>({ Errors::InvalidState() });
 		}
 	}
@@ -226,7 +226,7 @@ namespace UE::Online {
 		}
 		else
 		{
-			UE_LOG(LogTemp, Verbose, TEXT("[%s] No data found for user [%s]."), UTF8_TO_TCHAR(__FUNCTION__), *ToLogString(Params.LocalAccountId));
+			UE_LOG(LogOnlineServices, Verbose, TEXT("[%s] No data found for user [%s]."), UTF8_TO_TCHAR(__FUNCTION__), *ToLogString(Params.LocalAccountId));
 			return TOnlineResult<FIsPresenceSession>({ Errors::InvalidState() });
 		}
 	}
@@ -273,7 +273,7 @@ namespace UE::Online {
 		}
 		else
 		{
-			UE_LOG(LogTemp, Verbose, TEXT("[%s] No session found for name [%s]."), UTF8_TO_TCHAR(__FUNCTION__), *Params.LocalName.ToString());
+			UE_LOG(LogOnlineServices, Verbose, TEXT("[%s] No session found for name [%s]."), UTF8_TO_TCHAR(__FUNCTION__), *Params.LocalName.ToString());
 			return TOnlineResult<FGetMutableSessionByName>(Errors::NotFound());
 		}
 	}
@@ -288,7 +288,7 @@ namespace UE::Online {
 		}
 		else
 		{
-			UE_LOG(LogTemp, Verbose, TEXT("[%s] No session found for id [%s]."), UTF8_TO_TCHAR(__FUNCTION__), *ToLogString(Params.SessionId));
+			UE_LOG(LogOnlineServices, Verbose, TEXT("[%s] No session found for id [%s]."), UTF8_TO_TCHAR(__FUNCTION__), *ToLogString(Params.SessionId));
 			return TOnlineResult<FGetMutableSessionById>(Errors::NotFound());
 		}
 	}
@@ -306,13 +306,13 @@ namespace UE::Online {
 			}
 			else
 			{
-				UE_LOG(LogTemp, Verbose, TEXT("[%s] No invite found for invite id [%s]."), UTF8_TO_TCHAR(__FUNCTION__), *ToLogString(Params.SessionInviteId));
+				UE_LOG(LogOnlineServices, Verbose, TEXT("[%s] No invite found for invite id [%s]."), UTF8_TO_TCHAR(__FUNCTION__), *ToLogString(Params.SessionInviteId));
 				return TOnlineResult<FGetSessionInviteById>(Errors::NotFound());
 			}
 		}
 		else
 		{
-			UE_LOG(LogTemp, Verbose, TEXT("[%s] No invite data found for account id [%s]."), UTF8_TO_TCHAR(__FUNCTION__), *ToLogString(Params.LocalAccountId));
+			UE_LOG(LogOnlineServices, Verbose, TEXT("[%s] No invite data found for account id [%s]."), UTF8_TO_TCHAR(__FUNCTION__), *ToLogString(Params.LocalAccountId));
 			return TOnlineResult<FGetSessionInviteById>(Errors::InvalidState());
 		}
 	}
@@ -342,14 +342,14 @@ namespace UE::Online {
 #define CHECK_PARAMS_ID_HANDLE(IdHandle) \
 	if (!IdHandle.IsValid()) \
 	{ \
-		UE_LOG(LogTemp, Verbose, TEXT("[%s] Value of [%s] invalid."), UTF8_TO_TCHAR(__FUNCTION__), UTF8_TO_TCHAR(#IdHandle)); \
+		UE_LOG(LogOnlineServices, Verbose, TEXT("[%s] Value of [%s] invalid."), UTF8_TO_TCHAR(__FUNCTION__), UTF8_TO_TCHAR(#IdHandle)); \
 		return TOptional<FOnlineError>(Errors::InvalidParams()); \
 	} \
 
 #define CHECK_PARAMS_FNAME(Name) \
 	if (Name.IsNone()) \
 	{ \
-		UE_LOG(LogTemp, Verbose, TEXT("[%s] Value of [%s] invalid."), UTF8_TO_TCHAR(__FUNCTION__), UTF8_TO_TCHAR(#Name)); \
+		UE_LOG(LogOnlineServices, Verbose, TEXT("[%s] Value of [%s] invalid."), UTF8_TO_TCHAR(__FUNCTION__), UTF8_TO_TCHAR(#Name)); \
 		return TOptional<FOnlineError>(Errors::InvalidParams()); \
 	} \
 
@@ -357,7 +357,7 @@ namespace UE::Online {
 	TOnlineResult<FGetSessionByName> GetSessionByNameResult = GetSessionByName({ SessionName }); \
 	if (GetSessionByNameResult.IsError()) \
 	{ \
-		UE_LOG(LogTemp, Verbose, TEXT("[%s] Could not find session with name [%s]."), UTF8_TO_TCHAR(__FUNCTION__), *SessionName.ToString()); \
+		UE_LOG(LogOnlineServices, Verbose, TEXT("[%s] Could not find session with name [%s]."), UTF8_TO_TCHAR(__FUNCTION__), *SessionName.ToString()); \
 		return TOptional<FOnlineError>(Errors::InvalidState()); \
 	} \
 
@@ -427,7 +427,7 @@ namespace UE::Online {
 
 		if (Params.SessionSettings.NumMaxConnections == 0)
 		{
-			UE_LOG(LogTemp, Verbose, TEXT("[%s] NumMaxConnections invalid, can't be 0."), UTF8_TO_TCHAR(__FUNCTION__));
+			UE_LOG(LogOnlineServices, Verbose, TEXT("[%s] NumMaxConnections invalid, can't be 0."), UTF8_TO_TCHAR(__FUNCTION__));
 			return Errors::InvalidParams();
 		}
 
@@ -446,7 +446,7 @@ namespace UE::Online {
 		TOnlineResult<FGetSessionByName> GetSessionByNameResult = GetSessionByName({ Params.SessionName });
 		if (GetSessionByNameResult.IsOk())
 		{
-			UE_LOG(LogTemp, Verbose, TEXT("[%s] Another session with the name [%s] already exists."), UTF8_TO_TCHAR(__FUNCTION__), *Params.SessionName.ToString());
+			UE_LOG(LogOnlineServices, Verbose, TEXT("[%s] Another session with the name [%s] already exists."), UTF8_TO_TCHAR(__FUNCTION__), *Params.SessionName.ToString());
 			return TOptional<FOnlineError>(Errors::InvalidState());
 		}
 
@@ -457,7 +457,7 @@ namespace UE::Online {
 				TOnlineResult<FGetPresenceSession> GetPresenceSessionResult = GetPresenceSession({ Params.LocalAccountId });
 				if (GetPresenceSessionResult.IsOk())
 				{
-					UE_LOG(LogTemp, Verbose, TEXT("[%s] Could not create session with bPresenceEnabled set to true when another already exists [%s]."), UTF8_TO_TCHAR(__FUNCTION__), *Entry.Key.ToString());
+					UE_LOG(LogOnlineServices, Verbose, TEXT("[%s] Could not create session with bPresenceEnabled set to true when another already exists [%s]."), UTF8_TO_TCHAR(__FUNCTION__), *Entry.Key.ToString());
 					return Errors::InvalidState();
 				}
 			}
@@ -492,7 +492,7 @@ namespace UE::Online {
 		{
 			if (Params.Mutations.NumMaxConnections.GetValue() == 0)
 			{
-				UE_LOG(LogTemp, Verbose, TEXT("[%s] NumMaxConnections invalid, can't be 0."), UTF8_TO_TCHAR(__FUNCTION__));
+				UE_LOG(LogOnlineServices, Verbose, TEXT("[%s] NumMaxConnections invalid, can't be 0."), UTF8_TO_TCHAR(__FUNCTION__));
 				return Errors::InvalidParams();
 			}
 		}
@@ -564,7 +564,7 @@ namespace UE::Online {
 
 		if (Params.MaxResults == 0)
 		{
-			UE_LOG(LogTemp, Verbose, TEXT("[%s] MaxResults can't be 0"), UTF8_TO_TCHAR(__FUNCTION__));
+			UE_LOG(LogOnlineServices, Verbose, TEXT("[%s] MaxResults can't be 0"), UTF8_TO_TCHAR(__FUNCTION__));
 			return TOptional<FOnlineError>(Errors::InvalidParams());
 		}
 
@@ -590,7 +590,7 @@ namespace UE::Online {
 	{
 		if (CurrentSessionSearchPromisesUserMap.Contains(Params.LocalAccountId))
 		{
-			UE_LOG(LogTemp, Verbose, TEXT("[%s] Search already in progress"), UTF8_TO_TCHAR(__FUNCTION__));
+			UE_LOG(LogOnlineServices, Verbose, TEXT("[%s] Search already in progress"), UTF8_TO_TCHAR(__FUNCTION__));
 			return Errors::AlreadyPending();
 		}
 
@@ -662,7 +662,7 @@ namespace UE::Online {
 		{
 			if (GetSessionByNameResult.GetOkValue().Session->GetSessionId() != Params.SessionId)
 			{
-				UE_LOG(LogTemp, Verbose, TEXT("[%s] Another session with the name [%s] already exists."), UTF8_TO_TCHAR(__FUNCTION__), *Params.SessionName.ToString());
+				UE_LOG(LogOnlineServices, Verbose, TEXT("[%s] Another session with the name [%s] already exists."), UTF8_TO_TCHAR(__FUNCTION__), *Params.SessionName.ToString());
 				return TOptional<FOnlineError>(Errors::InvalidState());
 			}
 		}
@@ -670,7 +670,7 @@ namespace UE::Online {
 		TOnlineResult<FGetSessionById> GetSessionByIdResult = GetSessionById({ Params.SessionId });
 		if (GetSessionByIdResult.IsError())
 		{
-			UE_LOG(LogTemp, Verbose, TEXT("[%s] Session [%s] not found. Please call FindSessions to get an updated list of available sessions "), UTF8_TO_TCHAR(__FUNCTION__), *ToLogString(Params.SessionId));
+			UE_LOG(LogOnlineServices, Verbose, TEXT("[%s] Session [%s] not found. Please call FindSessions to get an updated list of available sessions "), UTF8_TO_TCHAR(__FUNCTION__), *ToLogString(Params.SessionId));
 			return TOptional<FOnlineError>(GetSessionByIdResult.GetErrorValue());
 		}
 
@@ -680,13 +680,13 @@ namespace UE::Online {
 
 		if (FoundSession->GetSessionMembers().Contains(Params.LocalAccountId))
 		{
-			UE_LOG(LogTemp, Verbose, TEXT("[%s] User [%s] already in session [%s]"), UTF8_TO_TCHAR(__FUNCTION__), *ToLogString(Params.LocalAccountId), *ToLogString(Params.SessionId));
+			UE_LOG(LogOnlineServices, Verbose, TEXT("[%s] User [%s] already in session [%s]"), UTF8_TO_TCHAR(__FUNCTION__), *ToLogString(Params.LocalAccountId), *ToLogString(Params.SessionId));
 			return TOptional<FOnlineError>(Errors::AccessDenied());
 		}
 
 		if (!FoundSession->IsJoinable())
 		{
-			UE_LOG(LogTemp, Verbose, TEXT("[%s] Session [%s] not joinable "), UTF8_TO_TCHAR(__FUNCTION__), *ToLogString(Params.SessionId));
+			UE_LOG(LogOnlineServices, Verbose, TEXT("[%s] Session [%s] not joinable "), UTF8_TO_TCHAR(__FUNCTION__), *ToLogString(Params.SessionId));
 			return TOptional<FOnlineError>(Errors::AccessDenied());
 		}
 
@@ -697,7 +697,7 @@ namespace UE::Online {
 				TOnlineResult<FGetPresenceSession> GetPresenceSessionResult = GetPresenceSession({ Params.LocalAccountId });
 				if (GetPresenceSessionResult.IsOk())
 				{
-					UE_LOG(LogTemp, Verbose, TEXT("[%s] Could not join session with bPresenceEnabled set to true when another presence session already exists [%s]."), UTF8_TO_TCHAR(__FUNCTION__), *Entry.Key.ToString());
+					UE_LOG(LogOnlineServices, Verbose, TEXT("[%s] Could not join session with bPresenceEnabled set to true when another presence session already exists [%s]."), UTF8_TO_TCHAR(__FUNCTION__), *Entry.Key.ToString());
 					return TOptional<FOnlineError>(Errors::InvalidState());
 				}
 			}
@@ -823,7 +823,7 @@ namespace UE::Online {
 
 		if (Params.TargetUsers.IsEmpty())
 		{
-			UE_LOG(LogTemp, Verbose, TEXT("[%s] TargetUsers array is empty"), UTF8_TO_TCHAR(__FUNCTION__));
+			UE_LOG(LogOnlineServices, Verbose, TEXT("[%s] TargetUsers array is empty"), UTF8_TO_TCHAR(__FUNCTION__));
 			return TOptional<FOnlineError>(Errors::InvalidParams());
 		}
 
@@ -868,13 +868,13 @@ namespace UE::Online {
 		{
 			if (!UserMap->Contains(Params.SessionInviteId))
 			{
-				UE_LOG(LogTemp, Verbose, TEXT("[%s] No invite found for invite id [%s]."), UTF8_TO_TCHAR(__FUNCTION__), *ToLogString(Params.SessionInviteId));
+				UE_LOG(LogOnlineServices, Verbose, TEXT("[%s] No invite found for invite id [%s]."), UTF8_TO_TCHAR(__FUNCTION__), *ToLogString(Params.SessionInviteId));
 				return TOptional<FOnlineError>(Errors::NotFound());
 			}
 		}
 		else
 		{
-			UE_LOG(LogTemp, Verbose, TEXT("[%s] No invite data found for account id [%s]."), UTF8_TO_TCHAR(__FUNCTION__), *ToLogString(Params.LocalAccountId));
+			UE_LOG(LogOnlineServices, Verbose, TEXT("[%s] No invite data found for account id [%s]."), UTF8_TO_TCHAR(__FUNCTION__), *ToLogString(Params.LocalAccountId));
 			return TOptional<FOnlineError>(Errors::InvalidState());
 		}
 
diff --git a/Engine/Plugins/Online/OnlineServices/Source/OnlineServicesCommon/Public/Online/OnlineServicesLog.h b/Engine/Plugins/Online/OnlineServices/Source/OnlineServicesCommon/Public/Online/OnlineServicesLog.h
deleted file mode 100644
diff --git a/Engine/Plugins/Online/OnlineServices/Source/OnlineServicesCommonEngineUtils/Private/Online/SessionsLAN.cpp b/Engine/Plugins/Online/OnlineServices/Source/OnlineServicesCommonEngineUtils/Private/Online/SessionsLAN.cpp
--- a/Engine/Plugins/Online/OnlineServices/Source/OnlineServicesCommonEngineUtils/Private/Online/SessionsLAN.cpp
+++ b/Engine/Plugins/Online/OnlineServices/Source/OnlineServicesCommonEngineUtils/Private/Online/SessionsLAN.cpp
@@ -102,7 +102,7 @@ const FSessionLAN& FSessionLAN::Cast(const ISession& InSession)
 void FSessionLAN::DumpState() const
 {
 	FSessionCommon::DumpState();
-	UE_LOG(LogTemp, Log, TEXT("OwnerInternetAddr [%s]"), *OwnerInternetAddr->ToString(true));
+	UE_LOG(LogOnlineServices, Log, TEXT("OwnerInternetAddr [%s]"), *OwnerInternetAddr->ToString(true));
 }
 
 /** FSessionsLAN */
@@ -227,7 +227,7 @@ TFuture<TOnlineResult<FJoinSession>> FSessionsLAN::JoinSessionImpl(const FJoinSe
 		// If no result is found, the id might be expired, which we should notify
 		if (FOnlineSessionIdRegistryLAN::GetChecked(Services.GetServicesProvider()).IsSessionIdExpired(Params.SessionId))
 		{
-			UE_LOG(LogTemp, Warning, TEXT("[%s] SessionId parameter [%s] is expired. Please call FindSessions to get an updated list of available sessions "), UTF8_TO_TCHAR(__FUNCTION__), *ToLogString(Params.SessionId));
+			UE_LOG(LogOnlineServices, Warning, TEXT("[%s] SessionId parameter [%s] is expired. Please call FindSessions to get an updated list of available sessions "), UTF8_TO_TCHAR(__FUNCTION__), *ToLogString(Params.SessionId));
 		}
 
 		Promise.EmplaceValue(GetSessionByIdResult.GetErrorValue());
@@ -292,13 +292,13 @@ TOptional<FOnlineError> FSessionsLAN::TryHostLANSession()
 
 		if (LANSessionManager->Host(QueryPacketDelegate))
 		{
-			UE_LOG(LogTemp, VeryVerbose, TEXT("[FSessionsLAN::TryHostLANSession] LAN Beacon hosting..."));
+			UE_LOG(LogOnlineServices, VeryVerbose, TEXT("[FSessionsLAN::TryHostLANSession] LAN Beacon hosting..."));
 
 			PublicSessionsHosted++;
 		}
 		else
 		{
-			UE_LOG(LogTemp, VeryVerbose, TEXT("[FSessionsLAN::TryHostLANSession] LAN Beacon failed to host!"));
+			UE_LOG(LogOnlineServices, VeryVerbose, TEXT("[FSessionsLAN::TryHostLANSession] LAN Beacon failed to host!"));
 
 			LANSessionManager->StopLANSession();
 
@@ -307,7 +307,7 @@ TOptional<FOnlineError> FSessionsLAN::TryHostLANSession()
 	}
 	else
 	{
-		UE_LOG(LogTemp, VeryVerbose, TEXT("[FSessionsLAN::TryHostLANSession] LAN Beacon already in use!"));
+		UE_LOG(LogOnlineServices, VeryVerbose, TEXT("[FSessionsLAN::TryHostLANSession] LAN Beacon already in use!"));
 
 		Result.Emplace(Errors::AlreadyPending());
 	}
@@ -328,7 +328,7 @@ void FSessionsLAN::FindLANSessions(const FAccountId& LocalAccountId)
 	LANSessionManager->CreateClientQueryPacket(Packet, LANSessionManager->LanNonce);
 	if (LANSessionManager->Search(Packet, ResponseDelegate, TimeoutDelegate) == false)
 	{
-		UE_LOG(LogTemp, Warning, TEXT("[FSessionsLAN::FindLANSessions] Search failed!"));
+		UE_LOG(LogOnlineServices, Warning, TEXT("[FSessionsLAN::FindLANSessions] Search failed!"));
 
 		if (LANSessionManager->GetBeaconState() == ELanBeaconState::Searching)
 		{
@@ -347,7 +347,7 @@ void FSessionsLAN::FindLANSessions(const FAccountId& LocalAccountId)
 	}
 	else
 	{
-		UE_LOG(LogTemp, Warning, TEXT("[]FSessionsLAN::FindLANSessions] Searching...."));
+		UE_LOG(LogOnlineServices, Warning, TEXT("[]FSessionsLAN::FindLANSessions] Searching...."));
 	}
 }
 
@@ -386,7 +386,7 @@ void FSessionsLAN::OnValidQueryPacketReceived(uint8* PacketData, int32 PacketLen
 			}
 			else
 			{
-				UE_LOG(LogTemp, Warning, TEXT("[FSessionsLAN::OnValidQueryPacketReceived] LAN broadcast packet overflow, cannot broadcast on LAN"));
+				UE_LOG(LogOnlineServices, Warning, TEXT("[FSessionsLAN::OnValidQueryPacketReceived] LAN broadcast packet overflow, cannot broadcast on LAN"));
 			}
 		}
 	}
@@ -409,7 +409,7 @@ void FSessionsLAN::OnLANSearchTimeout(const FAccountId LocalAccountId)
 		LANSessionManager->StopLANSession();
 	}
 
-	UE_LOG(LogTemp, Warning, TEXT("[FSessionsLAN::OnLANSearchTimeout] %d sessions found!"), SearchResultsUserMap.FindChecked(LocalAccountId).Num());
+	UE_LOG(LogOnlineServices, Warning, TEXT("[FSessionsLAN::OnLANSearchTimeout] %d sessions found!"), SearchResultsUserMap.FindChecked(LocalAccountId).Num());
 
 	TArray<FOnlineSessionId>& FoundSessionIds = SearchResultsUserMap.FindChecked(LocalAccountId);
 
diff --git a/Engine/Plugins/Online/OnlineServices/Source/OnlineServicesInterface/Private/Online/OnlineServices.cpp b/Engine/Plugins/Online/OnlineServices/Source/OnlineServicesInterface/Private/Online/OnlineServices.cpp
--- a/Engine/Plugins/Online/OnlineServices/Source/OnlineServicesInterface/Private/Online/OnlineServices.cpp
+++ b/Engine/Plugins/Online/OnlineServices/Source/OnlineServicesInterface/Private/Online/OnlineServices.cpp
@@ -2,11 +2,14 @@
 
 #include "Online/OnlineServices.h"
 
+#include "Online/OnlineServicesLog.h"
 #include "Online/OnlineServicesRegistry.h"
 #include "HAL/IConsoleManager.h"
 #include "Misc/CommandLine.h"
 #include "Misc/ConfigCacheIni.h"
 
+DEFINE_LOG_CATEGORY(LogOnlineServices);
+
 namespace UE::Online {
 
 int32 GetBuildUniqueId()
@@ -38,12 +41,12 @@ int32 GetBuildUniqueId()
 		{
 			if (!GConfig->GetBool(TEXT("OnlineServices"), TEXT("bUseBuildIdOverride"), bUseBuildIdOverride, GEngineIni))
 			{
-				UE_LOG(LogTemp, Warning, TEXT("Missing bUseBuildIdOverride= in [OnlineServices] of DefaultEngine.ini"));
+				UE_LOG(LogOnlineServices, Warning, TEXT("Missing bUseBuildIdOverride= in [OnlineServices] of DefaultEngine.ini"));
 			}
 
 			if (!GConfig->GetInt(TEXT("OnlineServices"), TEXT("BuildIdOverride"), BuildIdOverride, GEngineIni))
 			{
-				UE_LOG(LogTemp, Warning, TEXT("Missing BuildIdOverride= in [OnlineServices] of DefaultEngine.ini"));
+				UE_LOG(LogOnlineServices, Warning, TEXT("Missing BuildIdOverride= in [OnlineServices] of DefaultEngine.ini"));
 			}
 		}
 
diff --git a/Engine/Plugins/Online/OnlineServices/Source/OnlineServicesInterface/Public/Online/OnlineServicesLog.h b/Engine/Plugins/Online/OnlineServices/Source/OnlineServicesInterface/Public/Online/OnlineServicesLog.h
new file mode 100644
--- /dev/null
+++ b/Engine/Plugins/Online/OnlineServices/Source/OnlineServicesInterface/Public/Online/OnlineServicesLog.h
@@ -0,0 +1,7 @@
+// Copyright Epic Games, Inc. All Rights Reserved.
+
+#pragma once
+
+#include "Logging/LogMacros.h"
+
+ONLINESERVICESINTERFACE_API DECLARE_LOG_CATEGORY_EXTERN(LogOnlineServices, Log, All);
diff --git a/Engine/Plugins/Online/OnlineServicesEOS/Source/Private/Online/CommerceEOS.cpp b/Engine/Plugins/Online/OnlineServicesEOS/Source/Private/Online/CommerceEOS.cpp
--- a/Engine/Plugins/Online/OnlineServicesEOS/Source/Private/Online/CommerceEOS.cpp
+++ b/Engine/Plugins/Online/OnlineServicesEOS/Source/Private/Online/CommerceEOS.cpp
@@ -226,7 +226,7 @@ TOnlineAsyncOpHandle<FCommerceCheckout> FCommerceEOS::Checkout(FCommerceCheckout
 		EOS_EResult Result = Data->ResultCode;
 		if (Result != EOS_EResult::EOS_Success)
 		{
-			UE_LOG(LogTemp, Error, TEXT("EOS_Ecom_Checkout: failed with error (%s)"), ANSI_TO_TCHAR(EOS_EResult_ToString(Data->ResultCode)));
+			UE_LOG(LogOnlineServices, Error, TEXT("EOS_Ecom_Checkout: failed with error (%s)"), ANSI_TO_TCHAR(EOS_EResult_ToString(Data->ResultCode)));
 			Op.SetError(Errors::FromEOSResult(Result));
 			return;
 		}
@@ -271,7 +271,7 @@ TOnlineAsyncOpHandle<FCommerceQueryTransactionEntitlements> FCommerceEOS::QueryT
 			EOS_EResult CopyTransactionResult = EOS_Ecom_CopyTransactionById(EcomHandle, &CopyTransactionOptions, &OutTransaction);
 			if (CopyTransactionResult != EOS_EResult::EOS_Success)
 			{
-				UE_LOG(LogTemp, Error, TEXT("EOS_Ecom_CopyTransactionById: failed with error (%s)"), ANSI_TO_TCHAR(EOS_EResult_ToString(CopyTransactionResult)));
+				UE_LOG(LogOnlineServices, Error, TEXT("EOS_Ecom_CopyTransactionById: failed with error (%s)"), ANSI_TO_TCHAR(EOS_EResult_ToString(CopyTransactionResult)));
 				Op.SetError(Errors::FromEOSResult(CopyTransactionResult));
 				return;
 			}
@@ -350,7 +350,7 @@ TOnlineAsyncOpHandle<FCommerceQueryEntitlements> FCommerceEOS::QueryEntitlements
 			const EOS_EResult Result = Data->ResultCode;
 			if (Result != EOS_EResult::EOS_Success)
 			{
-				UE_LOG(LogTemp, Error, TEXT("EOS_Ecom_QueryEntitlements: failed with error (%s)"), ANSI_TO_TCHAR(EOS_EResult_ToString(Data->ResultCode)));
+				UE_LOG(LogOnlineServices, Error, TEXT("EOS_Ecom_QueryEntitlements: failed with error (%s)"), ANSI_TO_TCHAR(EOS_EResult_ToString(Data->ResultCode)));
 				Op.SetError(Errors::FromEOSResult(Result));
 				return;
 			}
@@ -376,7 +376,7 @@ TOnlineAsyncOpHandle<FCommerceQueryEntitlements> FCommerceEOS::QueryEntitlements
 				EOS_EResult CopyResult = EOS_Ecom_CopyEntitlementByIndex(EcomHandle, &CopyOptions, &EosEntitlement);
 				if (CopyResult != EOS_EResult::EOS_Success && CopyResult != EOS_EResult::EOS_Ecom_EntitlementStale)
 				{
-					UE_LOG(LogTemp, Error, TEXT("EOS_Ecom_CopyEntitlementByIndex: failed with error (%s) (proceeding with operation)"), ANSI_TO_TCHAR(EOS_EResult_ToString(CopyResult)));
+					UE_LOG(LogOnlineServices, Error, TEXT("EOS_Ecom_CopyEntitlementByIndex: failed with error (%s) (proceeding with operation)"), ANSI_TO_TCHAR(EOS_EResult_ToString(CopyResult)));
 					continue;
 				}
 
@@ -439,7 +439,7 @@ TOnlineAsyncOpHandle<FCommerceRedeemEntitlement> FCommerceEOS::RedeemEntitlement
 		EOS_EResult Result = Data->ResultCode;
 		if (Result != EOS_EResult::EOS_Success)
 		{
-			UE_LOG(LogTemp, Error, TEXT("EOS_Ecom_RedeemEntitlements: failed with error (%s)"), ANSI_TO_TCHAR(EOS_EResult_ToString(Data->ResultCode)));
+			UE_LOG(LogOnlineServices, Error, TEXT("EOS_Ecom_RedeemEntitlements: failed with error (%s)"), ANSI_TO_TCHAR(EOS_EResult_ToString(Data->ResultCode)));
 			Op.SetError(Errors::FromEOSResult(Result));
 			return;
 		}
@@ -484,7 +484,7 @@ TOnlineAsyncOpHandle<FCommerceRetrieveS2SToken> FCommerceEOS::RetrieveS2SToken(F
 			EOS_EResult Result = Data->ResultCode;
 			if (Result != EOS_EResult::EOS_Success)
 			{
-				UE_LOG(LogTemp, Error, TEXT("EOS_Ecom_QueryOwnershipToken: failed with error (%s)"), ANSI_TO_TCHAR(EOS_EResult_ToString(Data->ResultCode)));
+				UE_LOG(LogOnlineServices, Error, TEXT("EOS_Ecom_QueryOwnershipToken: failed with error (%s)"), ANSI_TO_TCHAR(EOS_EResult_ToString(Data->ResultCode)));
 				Op.SetError(Errors::FromEOSResult(Result));
 				return;
 			}
diff --git a/Engine/Plugins/Online/OnlineServicesEOS/Source/Private/Online/ExternalUIEOS.cpp b/Engine/Plugins/Online/OnlineServicesEOS/Source/Private/Online/ExternalUIEOS.cpp
--- a/Engine/Plugins/Online/OnlineServicesEOS/Source/Private/Online/ExternalUIEOS.cpp
+++ b/Engine/Plugins/Online/OnlineServicesEOS/Source/Private/Online/ExternalUIEOS.cpp
@@ -82,7 +82,7 @@ TOnlineAsyncOpHandle<FExternalUIShowFriendsUI> FExternalUIEOS::ShowFriendsUI(FEx
 	EOS_EpicAccountId LocalUserEasId = GetEpicAccountId(Params.LocalAccountId);
 	if (!EOS_EpicAccountId_IsValid(LocalUserEasId))
 	{
-		UE_LOG(LogTemp, Warning, TEXT("[FExternalUIEOS::ShowFriendsUI] LocalAccountId=[%s] EpicAccountId not found"), *ToLogString(Params.LocalAccountId));
+		UE_LOG(LogOnlineServices, Warning, TEXT("[FExternalUIEOS::ShowFriendsUI] LocalAccountId=[%s] EpicAccountId not found"), *ToLogString(Params.LocalAccountId));
 		Op->SetError(Errors::Unknown()); // TODO
 		return Op->GetHandle();
 	}
@@ -98,7 +98,7 @@ TOnlineAsyncOpHandle<FExternalUIShowFriendsUI> FExternalUIEOS::ShowFriendsUI(FEx
 	})
 	.Then([this](TOnlineAsyncOp<FExternalUIShowFriendsUI>& InAsyncOp, const EOS_UI_ShowFriendsCallbackInfo* Data)
 	{
-		UE_LOG(LogTemp, Verbose, TEXT("[FExternalUIEOS::ShowFriendsUI] EOS_UI_ShowFriends Result [%s] for User [%s]."), *LexToString(Data->ResultCode), *LexToString(Data->LocalUserId));
+		UE_LOG(LogOnlineServices, Verbose, TEXT("[FExternalUIEOS::ShowFriendsUI] EOS_UI_ShowFriends Result [%s] for User [%s]."), *LexToString(Data->ResultCode), *LexToString(Data->LocalUserId));
 
 		if (Data->ResultCode == EOS_EResult::EOS_Success)
 		{
diff --git a/Engine/Plugins/Online/OnlineServicesEOS/Source/Private/Online/PresenceEOS.cpp b/Engine/Plugins/Online/OnlineServicesEOS/Source/Private/Online/PresenceEOS.cpp
--- a/Engine/Plugins/Online/OnlineServicesEOS/Source/Private/Online/PresenceEOS.cpp
+++ b/Engine/Plugins/Online/OnlineServicesEOS/Source/Private/Online/PresenceEOS.cpp
@@ -85,13 +85,13 @@ void FPresenceEOS::Initialize()
 			This->Services.Get<FAuthEOS>()->ResolveAccountId(LocalAccountId, Data->PresenceUserId)
 				.Next([This, LocalAccountId](const FAccountId& PresenceAccountId)
 					{
-						UE_LOG(LogTemp, Verbose, TEXT("OnEOSPresenceUpdate: LocalAccountId=[%s] PresenceAccountId=[%s]"), *ToLogString(LocalAccountId), *ToLogString(PresenceAccountId));
+						UE_LOG(LogOnlineServices, Verbose, TEXT("OnEOSPresenceUpdate: LocalAccountId=[%s] PresenceAccountId=[%s]"), *ToLogString(LocalAccountId), *ToLogString(PresenceAccountId));
 						This->UpdateUserPresence(LocalAccountId, PresenceAccountId);
 					});
 		}
 		else // In some cases, this delegate will fire before the login process has completed, so we won't be able to find the AccountId just yet. We'll queue that presence update call for after Login has completed
 		{
-			UE_LOG(LogTemp, Verbose, TEXT("OnEOSPresenceUpdate: Account id not found for Epic id [%s]. Will retry after login completes]"), *LexToString(Data->LocalUserId));
+			UE_LOG(LogOnlineServices, Verbose, TEXT("OnEOSPresenceUpdate: Account id not found for Epic id [%s]. Will retry after login completes]"), *LexToString(Data->LocalUserId));
 
 			TArray<EOS_EpicAccountId>& PendingPresenceUpdateArray = This->PendingPresenceUpdates.FindOrAdd(Data->LocalUserId);
 			PendingPresenceUpdateArray.Add(Data->PresenceUserId);
@@ -123,7 +123,7 @@ void FPresenceEOS::HandleAuthLoginStatusChanged(const FAuthLoginStatusChanged& E
 				Services.Get<FAuthEOS>()->ResolveAccountId(LocalAccountId, PresenceAccountId)
 					.Next([this, LocalAccountId](const FAccountId& PresenceAccountId)
 						{
-							UE_LOG(LogTemp, Verbose, TEXT("OnEOSPresenceUpdate: LocalAccountId=[%s] PresenceAccountId=[%s]"), *ToLogString(LocalAccountId), *ToLogString(PresenceAccountId));
+							UE_LOG(LogOnlineServices, Verbose, TEXT("OnEOSPresenceUpdate: LocalAccountId=[%s] PresenceAccountId=[%s]"), *ToLogString(LocalAccountId), *ToLogString(PresenceAccountId));
 							UpdateUserPresence(LocalAccountId, PresenceAccountId);
 						});
 			}
@@ -173,7 +173,7 @@ TOnlineAsyncOpHandle<FQueryPresence> FPresenceEOS::QueryPresence(FQueryPresence:
 			})
 			.Then([this](TOnlineAsyncOp<FQueryPresence>& InAsyncOp, const EOS_Presence_QueryPresenceCallbackInfo* Data) mutable
 			{
-				UE_LOG(LogTemp, Warning, TEXT("QueryPresenceResult: [%s]"), *LexToString(Data->ResultCode));
+				UE_LOG(LogOnlineServices, Warning, TEXT("QueryPresenceResult: [%s]"), *LexToString(Data->ResultCode));
 
 				if (Data->ResultCode == EOS_EResult::EOS_Success)
 				{
@@ -238,7 +238,7 @@ TOnlineAsyncOpHandle<FUpdatePresence> FPresenceEOS::UpdatePresence(FUpdatePresen
 			EOS_EResult SetStatusResult = EOS_PresenceModification_SetStatus(ChangeHandle, &StatusOptions);
 			if (SetStatusResult != EOS_EResult::EOS_Success)
 			{
-				UE_LOG(LogTemp, Warning, TEXT("UpdatePresence: EOS_PresenceModification_SetStatus failed with result %s"), *LexToString(SetStatusResult));
+				UE_LOG(LogOnlineServices, Warning, TEXT("UpdatePresence: EOS_PresenceModification_SetStatus failed with result %s"), *LexToString(SetStatusResult));
 				InAsyncOp.SetError(Errors::FromEOSResult(SetStatusResult));
 				Promise.EmplaceValue();
 				return;
@@ -256,7 +256,7 @@ TOnlineAsyncOpHandle<FUpdatePresence> FPresenceEOS::UpdatePresence(FUpdatePresen
 			EOS_EResult SetRawRichTextResult = EOS_PresenceModification_SetRawRichText(ChangeHandle, &RawRichTextOptions);
 			if (SetRawRichTextResult != EOS_EResult::EOS_Success)
 			{
-				UE_LOG(LogTemp, Warning, TEXT("UpdatePresence: EOS_PresenceModification_SetRawRichText failed with result %s"), *LexToString(SetRawRichTextResult));
+				UE_LOG(LogOnlineServices, Warning, TEXT("UpdatePresence: EOS_PresenceModification_SetRawRichText failed with result %s"), *LexToString(SetRawRichTextResult));
 				InAsyncOp.SetError(Errors::FromEOSResult(SetRawRichTextResult));
 				Promise.EmplaceValue();
 				return;
@@ -289,7 +289,7 @@ TOnlineAsyncOpHandle<FUpdatePresence> FPresenceEOS::UpdatePresence(FUpdatePresen
 					UE_EOS_CHECK_API_MISMATCH(EOS_PRESENCEMODIFICATION_DATARECORDID_API_LATEST, 1);
 					RecordId.Key = Utf8Key.Get();
 
-					UE_LOG(LogTemp, Warning, TEXT("UpdatePresence: Removing field %s"), *RemovedProperty); // Temp logging
+					UE_LOG(LogOnlineServices, Warning, TEXT("UpdatePresence: Removing field %s"), *RemovedProperty); // Temp logging
 				}
 
 				EOS_PresenceModification_DeleteDataOptions DataOptions = { };
@@ -300,7 +300,7 @@ TOnlineAsyncOpHandle<FUpdatePresence> FPresenceEOS::UpdatePresence(FUpdatePresen
 				EOS_EResult DeleteDataResult = EOS_PresenceModification_DeleteData(ChangeHandle, &DataOptions);
 				if (DeleteDataResult != EOS_EResult::EOS_Success)
 				{
-					UE_LOG(LogTemp, Warning, TEXT("UpdatePresence: EOS_PresenceModification_DeleteDataOptions failed with result %s"), *LexToString(DeleteDataResult));
+					UE_LOG(LogOnlineServices, Warning, TEXT("UpdatePresence: EOS_PresenceModification_DeleteDataOptions failed with result %s"), *LexToString(DeleteDataResult));
 					InAsyncOp.SetError(Errors::FromEOSResult(DeleteDataResult));
 					Promise.EmplaceValue();
 					return;
@@ -313,7 +313,7 @@ TOnlineAsyncOpHandle<FUpdatePresence> FPresenceEOS::UpdatePresence(FUpdatePresen
 				if (Params.Presence->Properties.Num() > EOS_PRESENCE_DATA_MAX_KEYS)
 				{
 					// TODO: Move this check higher.  Needs to take into account number of present fields (not just ones updated) and removed fields.
-					UE_LOG(LogTemp, Warning, TEXT("UpdatePresence: Too many presence keys.  %u/%u"), Params.Presence->Properties.Num(), EOS_PRESENCE_DATA_MAX_KEYS);
+					UE_LOG(LogOnlineServices, Warning, TEXT("UpdatePresence: Too many presence keys.  %u/%u"), Params.Presence->Properties.Num(), EOS_PRESENCE_DATA_MAX_KEYS);
 					InAsyncOp.SetError(Errors::InvalidParams());
 					Promise.EmplaceValue();
 					return;
@@ -331,7 +331,7 @@ TOnlineAsyncOpHandle<FUpdatePresence> FPresenceEOS::UpdatePresence(FUpdatePresen
 					UE_EOS_CHECK_API_MISMATCH(EOS_PRESENCE_DATARECORD_API_LATEST, 1);
 					Record.Key = Utf8Key.Get();
 					Record.Value = Utf8Value.Get();
-					UE_LOG(LogTemp, Warning, TEXT("UpdatePresence: Set field [%s] to [%s]"), *UpdatedProperty.Key, *UpdatedProperty.Value); // Temp logging
+					UE_LOG(LogOnlineServices, Warning, TEXT("UpdatePresence: Set field [%s] to [%s]"), *UpdatedProperty.Key, *UpdatedProperty.Value); // Temp logging
 				}
 
 				EOS_PresenceModification_SetDataOptions DataOptions = { };
@@ -342,7 +342,7 @@ TOnlineAsyncOpHandle<FUpdatePresence> FPresenceEOS::UpdatePresence(FUpdatePresen
 				EOS_EResult SetDataResult = EOS_PresenceModification_SetData(ChangeHandle, &DataOptions);
 				if (SetDataResult != EOS_EResult::EOS_Success)
 				{
-					UE_LOG(LogTemp, Warning, TEXT("UpdatePresence: EOS_PresenceModification_SetData failed with result %s"), *LexToString(SetDataResult));
+					UE_LOG(LogOnlineServices, Warning, TEXT("UpdatePresence: EOS_PresenceModification_SetData failed with result %s"), *LexToString(SetDataResult));
 					InAsyncOp.SetError(Errors::FromEOSResult(SetDataResult));
 					Promise.EmplaceValue();
 					return;
@@ -367,7 +367,7 @@ TOnlineAsyncOpHandle<FUpdatePresence> FPresenceEOS::UpdatePresence(FUpdatePresen
 	})
 	.Then([this](TOnlineAsyncOp<FUpdatePresence>& InAsyncOp, const EOS_Presence_SetPresenceCallbackInfo* Data) mutable
 	{
-		UE_LOG(LogTemp, Warning, TEXT("SetPresenceResult: [%s]"), *LexToString(Data->ResultCode));
+		UE_LOG(LogOnlineServices, Warning, TEXT("SetPresenceResult: [%s]"), *LexToString(Data->ResultCode));
 
 		if (Data->ResultCode == EOS_EResult::EOS_Success)
 		{
@@ -431,14 +431,14 @@ TOnlineAsyncOpHandle<FPartialUpdatePresence> FPresenceEOS::PartialUpdatePresence
 						EOS_EResult SetStatusResult = EOS_PresenceModification_SetStatus(ChangeHandle, &StatusOptions);
 						if (SetStatusResult != EOS_EResult::EOS_Success)
 						{
-							UE_LOG(LogTemp, Warning, TEXT("UpdatePresence: EOS_PresenceModification_SetStatus failed with result %s"), *LexToString(SetStatusResult));
+							UE_LOG(LogOnlineServices, Warning, TEXT("UpdatePresence: EOS_PresenceModification_SetStatus failed with result %s"), *LexToString(SetStatusResult));
 							InAsyncOp.SetError(Errors::FromEOSResult(SetStatusResult));
 							Promise.EmplaceValue();
 							return;
 						}
 						else
 						{
-							UE_LOG(LogTemp, Warning, TEXT("UpdatePresence: Status set to %u"), (uint8)Params.Mutations.Status.GetValue()); // Temp logging
+							UE_LOG(LogOnlineServices, Warning, TEXT("UpdatePresence: Status set to %u"), (uint8)Params.Mutations.Status.GetValue()); // Temp logging
 						}
 					}
 
@@ -456,14 +456,14 @@ TOnlineAsyncOpHandle<FPartialUpdatePresence> FPresenceEOS::PartialUpdatePresence
 						EOS_EResult SetRawRichTextResult = EOS_PresenceModification_SetRawRichText(ChangeHandle, &RawRichTextOptions);
 						if (SetRawRichTextResult != EOS_EResult::EOS_Success)
 						{
-							UE_LOG(LogTemp, Warning, TEXT("UpdatePresence: EOS_PresenceModification_SetRawRichText failed with result %s"), *LexToString(SetRawRichTextResult));
+							UE_LOG(LogOnlineServices, Warning, TEXT("UpdatePresence: EOS_PresenceModification_SetRawRichText failed with result %s"), *LexToString(SetRawRichTextResult));
 							InAsyncOp.SetError(Errors::FromEOSResult(SetRawRichTextResult));
 							Promise.EmplaceValue();
 							return;
 						}
 						else
 						{
-							UE_LOG(LogTemp, Warning, TEXT("UpdatePresence: RichText set to %s"), *Params.Mutations.StatusString.GetValue()); // Temp logging
+							UE_LOG(LogOnlineServices, Warning, TEXT("UpdatePresence: RichText set to %s"), *Params.Mutations.StatusString.GetValue()); // Temp logging
 						}
 					}
 					// Removed fields
@@ -482,7 +482,7 @@ TOnlineAsyncOpHandle<FPartialUpdatePresence> FPresenceEOS::PartialUpdatePresence
 							UE_EOS_CHECK_API_MISMATCH(EOS_PRESENCEMODIFICATION_DATARECORDID_API_LATEST, 1);
 							RecordId.Key = Utf8Key.Get();
 
-							UE_LOG(LogTemp, Warning, TEXT("UpdatePresence: Removing field %s"), *RemovedProperty); // Temp logging
+							UE_LOG(LogOnlineServices, Warning, TEXT("UpdatePresence: Removing field %s"), *RemovedProperty); // Temp logging
 						}
 
 						EOS_PresenceModification_DeleteDataOptions DataOptions = { };
@@ -493,7 +493,7 @@ TOnlineAsyncOpHandle<FPartialUpdatePresence> FPresenceEOS::PartialUpdatePresence
 						EOS_EResult DeleteDataResult = EOS_PresenceModification_DeleteData(ChangeHandle, &DataOptions);
 						if (DeleteDataResult != EOS_EResult::EOS_Success)
 						{
-							UE_LOG(LogTemp, Warning, TEXT("UpdatePresence: EOS_PresenceModification_DeleteDataOptions failed with result %s"), *LexToString(DeleteDataResult));
+							UE_LOG(LogOnlineServices, Warning, TEXT("UpdatePresence: EOS_PresenceModification_DeleteDataOptions failed with result %s"), *LexToString(DeleteDataResult));
 							InAsyncOp.SetError(Errors::FromEOSResult(DeleteDataResult));
 							Promise.EmplaceValue();
 							return;
@@ -506,7 +506,7 @@ TOnlineAsyncOpHandle<FPartialUpdatePresence> FPresenceEOS::PartialUpdatePresence
 						if (Params.Mutations.UpdatedProperties.Num() > EOS_PRESENCE_DATA_MAX_KEYS)
 						{
 							// TODO: Move this check higher.  Needs to take into account number of present fields (not just ones updated) and removed fields.
-							UE_LOG(LogTemp, Warning, TEXT("UpdatePresence: Too many presence keys.  %u/%u"), Params.Mutations.UpdatedProperties.Num(), EOS_PRESENCE_DATA_MAX_KEYS);
+							UE_LOG(LogOnlineServices, Warning, TEXT("UpdatePresence: Too many presence keys.  %u/%u"), Params.Mutations.UpdatedProperties.Num(), EOS_PRESENCE_DATA_MAX_KEYS);
 							InAsyncOp.SetError(Errors::InvalidParams());
 							Promise.EmplaceValue();
 							return;
@@ -524,7 +524,7 @@ TOnlineAsyncOpHandle<FPartialUpdatePresence> FPresenceEOS::PartialUpdatePresence
 							UE_EOS_CHECK_API_MISMATCH(EOS_PRESENCE_DATARECORD_API_LATEST, 1);
 							Record.Key = Utf8Key.Get();
 							Record.Value = Utf8Value.Get();
-							UE_LOG(LogTemp, Warning, TEXT("UpdatePresence: Set field [%s] to [%s]"), *UpdatedProperty.Key, *UpdatedProperty.Value); // Temp logging
+							UE_LOG(LogOnlineServices, Warning, TEXT("UpdatePresence: Set field [%s] to [%s]"), *UpdatedProperty.Key, *UpdatedProperty.Value); // Temp logging
 						}
 
 						EOS_PresenceModification_SetDataOptions DataOptions = { };
@@ -535,7 +535,7 @@ TOnlineAsyncOpHandle<FPartialUpdatePresence> FPresenceEOS::PartialUpdatePresence
 						EOS_EResult SetDataResult = EOS_PresenceModification_SetData(ChangeHandle, &DataOptions);
 						if (SetDataResult != EOS_EResult::EOS_Success)
 						{
-							UE_LOG(LogTemp, Warning, TEXT("UpdatePresence: EOS_PresenceModification_SetData failed with result %s"), *LexToString(SetDataResult));
+							UE_LOG(LogOnlineServices, Warning, TEXT("UpdatePresence: EOS_PresenceModification_SetData failed with result %s"), *LexToString(SetDataResult));
 							InAsyncOp.SetError(Errors::FromEOSResult(SetDataResult));
 							Promise.EmplaceValue();
 							return;
@@ -559,7 +559,7 @@ TOnlineAsyncOpHandle<FPartialUpdatePresence> FPresenceEOS::PartialUpdatePresence
 			})
 			.Then([this](TOnlineAsyncOp<FPartialUpdatePresence>& InAsyncOp, const EOS_Presence_SetPresenceCallbackInfo* Data) mutable
 			{
-				UE_LOG(LogTemp, Warning, TEXT("SetPresenceResult: [%s]"), *LexToString(Data->ResultCode));
+				UE_LOG(LogOnlineServices, Warning, TEXT("SetPresenceResult: [%s]"), *LexToString(Data->ResultCode));
 
 				if (Data->ResultCode == EOS_EResult::EOS_Success)
 				{
@@ -690,7 +690,7 @@ void FPresenceEOS::UpdateUserPresence(FAccountId LocalAccountId, FAccountId Pres
 	}
 	else
 	{
-		UE_LOG(LogTemp, Error, TEXT("UpdateUserPresence: CopyPresence Failed %s"), *LexToString(CopyPresenceResult));
+		UE_LOG(LogOnlineServices, Error, TEXT("UpdateUserPresence: CopyPresence Failed %s"), *LexToString(CopyPresenceResult));
 	}
 
 	if (bPresenceHasChanged)
diff --git a/Engine/Plugins/Online/OnlineServicesEOS/Source/Private/Online/UserInfoEOS.cpp b/Engine/Plugins/Online/OnlineServicesEOS/Source/Private/Online/UserInfoEOS.cpp
--- a/Engine/Plugins/Online/OnlineServicesEOS/Source/Private/Online/UserInfoEOS.cpp
+++ b/Engine/Plugins/Online/OnlineServicesEOS/Source/Private/Online/UserInfoEOS.cpp
@@ -70,7 +70,7 @@ TOnlineAsyncOpHandle<FQueryUserInfo> FUserInfoEOS::QueryUserInfo(FQueryUserInfo:
 			{
 				if (CallbackInfo->ResultCode != EOS_EResult::EOS_Success)
 				{
-					UE_LOG(LogTemp, Warning, TEXT("EOS_UserInfo_QueryUserInfo failed with result=[%s]"), *LexToString(CallbackInfo->ResultCode));
+					UE_LOG(LogOnlineServices, Warning, TEXT("EOS_UserInfo_QueryUserInfo failed with result=[%s]"), *LexToString(CallbackInfo->ResultCode));
 					Op.SetError(Errors::FromEOSResult(CallbackInfo->ResultCode));
 				}
 			});
@@ -114,7 +114,7 @@ TOnlineResult<FGetUserInfo> FUserInfoEOS::GetUserInfo(FGetUserInfo::Params&& Par
 	const EOS_EResult EosResult = EOS_UserInfo_CopyUserInfo(UserInfoHandle, &Options, &EosUserInfo);
 	if(EosResult != EOS_EResult::EOS_Success)
 	{
-		UE_LOG(LogTemp, Warning, TEXT("EOS_UserInfo_CopyUserInfo failed with result=[%s]"), *LexToString(EosResult));
+		UE_LOG(LogOnlineServices, Warning, TEXT("EOS_UserInfo_CopyUserInfo failed with result=[%s]"), *LexToString(EosResult));
 		return TOnlineResult<FGetUserInfo>(Errors::FromEOSResult(EosResult));
 	}
 
diff --git a/Engine/Plugins/Online/OnlineServicesEOSGS/Source/Private/Online/AchievementsEOSGS.cpp b/Engine/Plugins/Online/OnlineServicesEOSGS/Source/Private/Online/AchievementsEOSGS.cpp
--- a/Engine/Plugins/Online/OnlineServicesEOSGS/Source/Private/Online/AchievementsEOSGS.cpp
+++ b/Engine/Plugins/Online/OnlineServicesEOSGS/Source/Private/Online/AchievementsEOSGS.cpp
@@ -39,7 +39,7 @@ void FAchievementsEOSGS::Initialize()
 			const FString AchievementId = UTF8_TO_TCHAR(Data->AchievementId);
 			const FDateTime UnlockTime = FDateTime::FromUnixTimestamp(Data->UnlockTime);
 
-			UE_LOG(LogTemp, Verbose, TEXT("EOS_Achievements_OnAchievementsUnlocked User=%s AchievId=%s"), *ToLogString(LocalAccountId), *AchievementId);
+			UE_LOG(LogOnlineServices, Verbose, TEXT("EOS_Achievements_OnAchievementsUnlocked User=%s AchievId=%s"), *ToLogString(LocalAccountId), *AchievementId);
 
 			if (FAchievementStateMap* LocalUserAchievementStates = This->AchievementStates.Find(LocalAccountId))
 			{
@@ -103,7 +103,7 @@ TOnlineAsyncOpHandle<FQueryAchievementDefinitions> FAchievementsEOSGS::QueryAchi
 		{
 			if (Data->ResultCode != EOS_EResult::EOS_Success)
 			{
-				UE_LOG(LogTemp, Warning, TEXT("EOS_Achievements_QueryDefinitions failed with result=[%s]"), *LexToString(Data->ResultCode));
+				UE_LOG(LogOnlineServices, Warning, TEXT("EOS_Achievements_QueryDefinitions failed with result=[%s]"), *LexToString(Data->ResultCode));
 
 				InAsyncOp.SetError(Errors::FromEOSResult(Data->ResultCode));
 				return;
@@ -127,7 +127,7 @@ TOnlineAsyncOpHandle<FQueryAchievementDefinitions> FAchievementsEOSGS::QueryAchi
 				const EOS_EResult CopyResult = EOS_Achievements_CopyAchievementDefinitionV2ByIndex(AchievementsHandle, &CopyOptions, &EOSDefinition);
 				if (CopyResult != EOS_EResult::EOS_Success)
 				{
-					UE_LOG(LogTemp, Warning, TEXT("EOS_Achievements_CopyAchievementDefinitionV2ByIndex failed with result=[%s]"), *LexToString(CopyResult));
+					UE_LOG(LogOnlineServices, Warning, TEXT("EOS_Achievements_CopyAchievementDefinitionV2ByIndex failed with result=[%s]"), *LexToString(CopyResult));
 
 					InAsyncOp.SetError(Errors::FromEOSResult(CopyResult));
 					return;
@@ -136,7 +136,7 @@ TOnlineAsyncOpHandle<FQueryAchievementDefinitions> FAchievementsEOSGS::QueryAchi
 				const int ExpectedDefinitionApiVersion = 2;
 				UE_EOS_CHECK_API_MISMATCH(EOS_ACHIEVEMENTS_DEFINITIONV2_API_LATEST, ExpectedDefinitionApiVersion);
 				const bool bDefsApiVersionOk = EOSDefinition->ApiVersion == ExpectedDefinitionApiVersion;
-				UE_CLOG(!bDefsApiVersionOk, LogTemp, Warning, TEXT("EOS_Achievements_DefinitionV2 version mismatch Expected=%d Actual=%d"), ExpectedDefinitionApiVersion, EOSDefinition->ApiVersion);
+				UE_CLOG(!bDefsApiVersionOk, LogOnlineServices, Warning, TEXT("EOS_Achievements_DefinitionV2 version mismatch Expected=%d Actual=%d"), ExpectedDefinitionApiVersion, EOSDefinition->ApiVersion);
 
 				FString AchievementId = UTF8_TO_TCHAR(EOSDefinition->AchievementId);
 				FAchievementDefinition& AchievementDefinition = NewAchievementDefinitions.Emplace(AchievementId);
@@ -156,7 +156,7 @@ TOnlineAsyncOpHandle<FQueryAchievementDefinitions> FAchievementsEOSGS::QueryAchi
 					const int ExpectedStatThresholdsApiVersion = 1;
 					UE_EOS_CHECK_API_MISMATCH(EOS_ACHIEVEMENTS_STATTHRESHOLDS_API_LATEST, ExpectedStatThresholdsApiVersion);
 					const bool bStatsApiVersionOk = EOSStatThreshold.ApiVersion == ExpectedStatThresholdsApiVersion;
-					UE_CLOG(!bStatsApiVersionOk, LogTemp, Warning, TEXT("EOS_Achievements_StatThresholds version mismatch Expected=%d Actual=%d"), ExpectedStatThresholdsApiVersion, EOSStatThreshold.ApiVersion);
+					UE_CLOG(!bStatsApiVersionOk, LogOnlineServices, Warning, TEXT("EOS_Achievements_StatThresholds version mismatch Expected=%d Actual=%d"), ExpectedStatThresholdsApiVersion, EOSStatThreshold.ApiVersion);
 
 					FAchievementStatDefinition& StatDefinition = AchievementDefinition.StatDefinitions.Emplace_GetRef();
 					StatDefinition.StatId = UTF8_TO_TCHAR(EOSStatThreshold.Name);
@@ -253,7 +253,7 @@ TOnlineAsyncOpHandle<FQueryAchievementStates> FAchievementsEOSGS::QueryAchieveme
 
 			if (Data->ResultCode != EOS_EResult::EOS_Success)
 			{
-				UE_LOG(LogTemp, Warning, TEXT("EOS_Achievements_QueryPlayerAchievements failed with result=[%s]"), *LexToString(Data->ResultCode));
+				UE_LOG(LogOnlineServices, Warning, TEXT("EOS_Achievements_QueryPlayerAchievements failed with result=[%s]"), *LexToString(Data->ResultCode));
 
 				InAsyncOp.SetError(Errors::FromEOSResult(Data->ResultCode));
 				return;
@@ -280,7 +280,7 @@ TOnlineAsyncOpHandle<FQueryAchievementStates> FAchievementsEOSGS::QueryAchieveme
 				const EOS_EResult CopyResult = EOS_Achievements_CopyPlayerAchievementByIndex(AchievementsHandle, &CopyOptions, &EOSPlayerAchievement);
 				if (CopyResult != EOS_EResult::EOS_Success)
 				{
-					UE_LOG(LogTemp, Warning, TEXT("EOS_Achievements_CopyPlayerAchievementByIndex failed with result=[%s]"), *LexToString(CopyResult));
+					UE_LOG(LogOnlineServices, Warning, TEXT("EOS_Achievements_CopyPlayerAchievementByIndex failed with result=[%s]"), *LexToString(CopyResult));
 
 					InAsyncOp.SetError(Errors::FromEOSResult(CopyResult));
 					return;
@@ -289,7 +289,7 @@ TOnlineAsyncOpHandle<FQueryAchievementStates> FAchievementsEOSGS::QueryAchieveme
 				const int ExpectedPlayerAchievementApiVersion = 2;
 				UE_EOS_CHECK_API_MISMATCH(EOS_ACHIEVEMENTS_PLAYERACHIEVEMENT_API_LATEST, ExpectedPlayerAchievementApiVersion);
 				const bool bPlayerAchievementApiVersionOk = EOSPlayerAchievement->ApiVersion == ExpectedPlayerAchievementApiVersion;
-				UE_CLOG(!bPlayerAchievementApiVersionOk, LogTemp, Warning, TEXT("EOS_Achievements_PlayerAchievement version mismatch Expected=%d Actual=%d"), ExpectedPlayerAchievementApiVersion, EOSPlayerAchievement->ApiVersion);
+				UE_CLOG(!bPlayerAchievementApiVersionOk, LogOnlineServices, Warning, TEXT("EOS_Achievements_PlayerAchievement version mismatch Expected=%d Actual=%d"), ExpectedPlayerAchievementApiVersion, EOSPlayerAchievement->ApiVersion);
 
 				FString AchievementId = UTF8_TO_TCHAR(EOSPlayerAchievement->AchievementId);
 				FAchievementState& AchievementState = NewAchievementStates.Emplace(AchievementId);
@@ -402,12 +402,12 @@ TOnlineAsyncOpHandle<FUnlockAchievements> FAchievementsEOSGS::UnlockAchievements
 		{
 			if (Data->ResultCode != EOS_EResult::EOS_Success)
 			{
-				UE_LOG(LogTemp, Warning, TEXT("EOS_Achievements_UnlockAchievements failed with result=[%s]"), *LexToString(Data->ResultCode));
+				UE_LOG(LogOnlineServices, Warning, TEXT("EOS_Achievements_UnlockAchievements failed with result=[%s]"), *LexToString(Data->ResultCode));
 				InAsyncOp.SetError(Errors::FromEOSResult(Data->ResultCode));
 				return;
 			}
 
-			UE_LOG(LogTemp, Verbose, TEXT("EOS_Achievements_UnlockAchievements succeeded"), *LexToString(Data->ResultCode));
+			UE_LOG(LogOnlineServices, Verbose, TEXT("EOS_Achievements_UnlockAchievements succeeded"), *LexToString(Data->ResultCode));
 
 			const FUnlockAchievements::Params& Params = InAsyncOp.GetParams();
 			const FDateTime TimeNow = FDateTime::Now();
diff --git a/Engine/Plugins/Online/OnlineServicesEOSGS/Source/Private/Online/LeaderboardsEOSGS.cpp b/Engine/Plugins/Online/OnlineServicesEOSGS/Source/Private/Online/LeaderboardsEOSGS.cpp
--- a/Engine/Plugins/Online/OnlineServicesEOSGS/Source/Private/Online/LeaderboardsEOSGS.cpp
+++ b/Engine/Plugins/Online/OnlineServicesEOSGS/Source/Private/Online/LeaderboardsEOSGS.cpp
@@ -78,7 +78,7 @@ TOnlineAsyncOpHandle<FReadEntriesForUsers> FLeaderboardsEOSGS::ReadEntriesForUse
 		{
 			if (Data->ResultCode != EOS_EResult::EOS_Success)
 			{
-				UE_LOG(LogTemp, Warning, TEXT("EOS_Leaderboards_IngestStat failed with result=[%s]"), *LexToString(Data->ResultCode));
+				UE_LOG(LogOnlineServices, Warning, TEXT("EOS_Leaderboards_IngestStat failed with result=[%s]"), *LexToString(Data->ResultCode));
 				InAsyncOp.SetError(Errors::FromEOSResult(Data->ResultCode));
 				return;
 			}
@@ -104,7 +104,7 @@ TOnlineAsyncOpHandle<FReadEntriesForUsers> FLeaderboardsEOSGS::ReadEntriesForUse
 
 				if (CopyResult != EOS_EResult::EOS_Success)
 				{
-					UE_LOG(LogTemp, VeryVerbose, TEXT("Value not found for leaderboard %s."), *Params.BoardName);
+					UE_LOG(LogOnlineServices, VeryVerbose, TEXT("Value not found for leaderboard %s."), *Params.BoardName);
 					continue;
 				}
 
@@ -191,7 +191,7 @@ TOnlineAsyncOpHandle<FReadEntriesAroundRank> FLeaderboardsEOSGS::ReadEntriesArou
 		{
 			if (Data->ResultCode != EOS_EResult::EOS_Success)
 			{
-				UE_LOG(LogTemp, Warning, TEXT("EOS_Leaderboards_QueryLeaderboardRanks failed with result=[%s]"), *LexToString(Data->ResultCode));
+				UE_LOG(LogOnlineServices, Warning, TEXT("EOS_Leaderboards_QueryLeaderboardRanks failed with result=[%s]"), *LexToString(Data->ResultCode));
 				InAsyncOp.SetError(Errors::FromEOSResult(Data->ResultCode));
 				return;
 			}
@@ -236,7 +236,7 @@ TOnlineAsyncOpHandle<FReadEntriesAroundUser> FLeaderboardsEOSGS::ReadEntriesArou
 		{
 			if (Data->ResultCode != EOS_EResult::EOS_Success)
 			{
-				UE_LOG(LogTemp, Warning, TEXT("EOS_Leaderboards_QueryLeaderboardRanks failed with result=[%s]"), *LexToString(Data->ResultCode));
+				UE_LOG(LogOnlineServices, Warning, TEXT("EOS_Leaderboards_QueryLeaderboardRanks failed with result=[%s]"), *LexToString(Data->ResultCode));
 				InAsyncOp.SetError(Errors::FromEOSResult(Data->ResultCode));
 				return;
 			}
diff --git a/Engine/Plugins/Online/OnlineServicesEOSGS/Source/Private/Online/LobbiesEOSGS.cpp b/Engine/Plugins/Online/OnlineServicesEOSGS/Source/Private/Online/LobbiesEOSGS.cpp
--- a/Engine/Plugins/Online/OnlineServicesEOSGS/Source/Private/Online/LobbiesEOSGS.cpp
+++ b/Engine/Plugins/Online/OnlineServicesEOSGS/Source/Private/Online/LobbiesEOSGS.cpp
@@ -85,7 +85,7 @@ TOnlineAsyncOpHandle<FCreateLobby> FLobbiesEOSGS::CreateLobby(FCreateLobby::Para
 		{
 			if (Result.IsError())
 			{
-				UE_LOG(LogTemp, Warning, TEXT("[FLobbiesEOSGS::CreateLobby] DestroyLobbyImpl Failed: User[%s], Lobby[%s], Result[%s]"),
+				UE_LOG(LogOnlineServices, Warning, TEXT("[FLobbiesEOSGS::CreateLobby] DestroyLobbyImpl Failed: User[%s], Lobby[%s], Result[%s]"),
 					*ToLogString(LocalAccountId), *LobbyIdString, *Result.GetErrorValue().GetLogString());
 			}
 
@@ -107,7 +107,7 @@ TOnlineAsyncOpHandle<FCreateLobby> FLobbiesEOSGS::CreateLobby(FCreateLobby::Para
 	// Check prerequisites.
 	if (!Services.GetAuthInterface()->IsLoggedIn(Params.LocalAccountId))
 	{
-		UE_LOG(LogTemp, Warning, TEXT("[FLobbiesEOSGS::CreateLobby] Failed: User is not logged in. User[%s]"), *ToLogString(Params.LocalAccountId));
+		UE_LOG(LogOnlineServices, Warning, TEXT("[FLobbiesEOSGS::CreateLobby] Failed: User is not logged in. User[%s]"), *ToLogString(Params.LocalAccountId));
 		Op->SetError(Errors::InvalidUser());
 		return Op->GetHandle();
 	}
@@ -115,7 +115,7 @@ TOnlineAsyncOpHandle<FCreateLobby> FLobbiesEOSGS::CreateLobby(FCreateLobby::Para
 	// Check that requested lobby schema exists.
 	if (!SchemaRegistry->GetDefinition(Params.SchemaId).IsValid())
 	{
-		UE_LOG(LogTemp, Warning, TEXT("[FLobbiesEOSGS::CreateLobby] Failed: Unknown lobby schema [%s]"), *Params.SchemaId.ToString().ToLower());
+		UE_LOG(LogOnlineServices, Warning, TEXT("[FLobbiesEOSGS::CreateLobby] Failed: Unknown lobby schema [%s]"), *Params.SchemaId.ToString().ToLower());
 		Op->SetError(Errors::InvalidParams());
 		return Op->GetHandle();
 	}
@@ -167,7 +167,7 @@ TOnlineAsyncOpHandle<FCreateLobby> FLobbiesEOSGS::CreateLobby(FCreateLobby::Para
 
 		if (Data->ResultCode != EOS_EResult::EOS_Success)
 		{
-			UE_LOG(LogTemp, Warning, TEXT("[FLobbiesEOSGS::CreateLobby] EOS_Lobby_CreateLobby Failed: User[%s], Result[%s]"),
+			UE_LOG(LogOnlineServices, Warning, TEXT("[FLobbiesEOSGS::CreateLobby] EOS_Lobby_CreateLobby Failed: User[%s], Result[%s]"),
 				*ToLogString(Params.LocalAccountId), *LexToString(Data->ResultCode));
 			InAsyncOp.SetError(Errors::FromEOSResult(Data->ResultCode));
 		}
@@ -188,7 +188,7 @@ TOnlineAsyncOpHandle<FCreateLobby> FLobbiesEOSGS::CreateLobby(FCreateLobby::Para
 			FLobbyDetailsEOS::CreateFromLobbyId(LobbyPrerequisites.ToSharedRef(), Params.LocalAccountId, LobbyIdString);
 		if (Result.IsError())
 		{
-			UE_LOG(LogTemp, Warning, TEXT("[FLobbiesEOSGS::CreateLobby] FLobbyDetailsEOS::CreateFromLobbyId Failed: User[%s], Lobby[%s], Result[%s]"),
+			UE_LOG(LogOnlineServices, Warning, TEXT("[FLobbiesEOSGS::CreateLobby] FLobbyDetailsEOS::CreateFromLobbyId Failed: User[%s], Lobby[%s], Result[%s]"),
 				*ToLogString(Params.LocalAccountId), *LobbyIdString, *Result.GetErrorValue().GetLogString());
 
 			return DestroyLobbyDuringCreate(InAsyncOp, Params.LocalAccountId, LobbyIdString, MoveTemp(Result.GetErrorValue()));
@@ -214,7 +214,7 @@ TOnlineAsyncOpHandle<FCreateLobby> FLobbiesEOSGS::CreateLobby(FCreateLobby::Para
 
 		if (Result.IsError())
 		{
-			UE_LOG(LogTemp, Warning, TEXT("[FLobbiesEOSGS::CreateLobby] FLobbyDataRegistryEOS::FindOrCreateFromLobbyDetails Failed: User[%s], Lobby[%s], Result[%s]"),
+			UE_LOG(LogOnlineServices, Warning, TEXT("[FLobbiesEOSGS::CreateLobby] FLobbyDataRegistryEOS::FindOrCreateFromLobbyDetails Failed: User[%s], Lobby[%s], Result[%s]"),
 				*ToLogString(Params.LocalAccountId), *LobbyIdString, *Result.GetErrorValue().GetLogString());
 
 			return DestroyLobbyDuringCreate(InAsyncOp, Params.LocalAccountId, LobbyIdString, MoveTemp(Result.GetErrorValue()));
@@ -242,7 +242,7 @@ TOnlineAsyncOpHandle<FCreateLobby> FLobbiesEOSGS::CreateLobby(FCreateLobby::Para
 		TOnlineResult<FLobbyClientDataPrepareClientChanges> PrepareResult = LobbyData->GetLobbyClientData()->PrepareClientChanges(MoveTemp(PrepareParams));
 		if (PrepareResult.IsError())
 		{
-			UE_LOG(LogTemp, Warning, TEXT("[FLobbiesEOSGS::CreateLobby] PrepareClientChanges Failed: User[%s], Lobby[%s], Result[%s]"),
+			UE_LOG(LogOnlineServices, Warning, TEXT("[FLobbiesEOSGS::CreateLobby] PrepareClientChanges Failed: User[%s], Lobby[%s], Result[%s]"),
 				*ToLogString(Params.LocalAccountId), *LobbyData->GetLobbyIdString(), *PrepareResult.GetErrorValue().GetLogString());
 			InAsyncOp.SetError(MoveTemp(PrepareResult.GetErrorValue()));
 			return MakeFulfilledPromise<TDefaultErrorResult<FLobbiesModifyLobbyDataImpl>>(FLobbiesModifyLobbyDataImpl::Result{}).GetFuture();
@@ -262,7 +262,7 @@ TOnlineAsyncOpHandle<FCreateLobby> FLobbiesEOSGS::CreateLobby(FCreateLobby::Para
 
 		if (Result.IsError())
 		{
-			UE_LOG(LogTemp, Warning, TEXT("[FLobbiesEOSGS::CreateLobby] ModifyLobbyDataImpl Failed: User[%s], Lobby[%s], Result[%s]"),
+			UE_LOG(LogOnlineServices, Warning, TEXT("[FLobbiesEOSGS::CreateLobby] ModifyLobbyDataImpl Failed: User[%s], Lobby[%s], Result[%s]"),
 				*ToLogString(Params.LocalAccountId), *LobbyData->GetLobbyIdString(), *Result.GetErrorValue().GetLogString());
 
 			return DestroyLobbyDuringCreate(InAsyncOp, Params.LocalAccountId, LobbyData->GetLobbyIdString(), MoveTemp(Result.GetErrorValue()));
@@ -284,7 +284,7 @@ TOnlineAsyncOpHandle<FCreateLobby> FLobbiesEOSGS::CreateLobby(FCreateLobby::Para
 		// Update local cache and fire events.
 		LobbyData->GetLobbyClientData()->CommitClientChanges({ &LobbyEvents });
 
-		UE_LOG(LogTemp, Log, TEXT("[FLobbiesEOSGS::CreateLobby] Succeeded: User[%s], Lobby[%s]"),
+		UE_LOG(LogOnlineServices, Log, TEXT("[FLobbiesEOSGS::CreateLobby] Succeeded: User[%s], Lobby[%s]"),
 			*ToLogString(Params.LocalAccountId), *LobbyData->GetLobbyIdString());
 		InAsyncOp.SetResult(FCreateLobby::Result{LobbyData->GetLobbyClientData()->GetPublicDataPtr()});
 	})
@@ -301,7 +301,7 @@ TOnlineAsyncOpHandle<FFindLobbies> FLobbiesEOSGS::FindLobbies(FFindLobbies::Para
 	// Check prerequisites.
 	if (!Services.GetAuthInterface()->IsLoggedIn(Params.LocalAccountId))
 	{
-		UE_LOG(LogTemp, Warning, TEXT("[FLobbiesEOSGS::FindLobbies] Failed: User is not logged in. User[%s]"), *ToLogString(Params.LocalAccountId));
+		UE_LOG(LogOnlineServices, Warning, TEXT("[FLobbiesEOSGS::FindLobbies] Failed: User is not logged in. User[%s]"), *ToLogString(Params.LocalAccountId));
 		Op->SetError(Errors::InvalidUser());
 		return Op->GetHandle();
 	}
@@ -320,7 +320,7 @@ TOnlineAsyncOpHandle<FFindLobbies> FLobbiesEOSGS::FindLobbies(FFindLobbies::Para
 		{
 			const FFindLobbies::Params& Params = InAsyncOp.GetParams();
 
-			UE_LOG(LogTemp, Warning, TEXT("[FLobbiesEOSGS::JoinLobby] FLobbySearchEOS::Create Failed: User[%s], Result[%s]"),
+			UE_LOG(LogOnlineServices, Warning, TEXT("[FLobbiesEOSGS::JoinLobby] FLobbySearchEOS::Create Failed: User[%s], Result[%s]"),
 				*ToLogString(Params.LocalAccountId), *Result.GetErrorValue().GetLogString());
 			InAsyncOp.SetError(MoveTemp(Result.GetErrorValue()));
 		}
@@ -336,7 +336,7 @@ TOnlineAsyncOpHandle<FFindLobbies> FLobbiesEOSGS::FindLobbies(FFindLobbies::Para
 
 		ActiveSearchResults.Add(InAsyncOp.GetParams().LocalAccountId, LobbySearch);
 
-		UE_LOG(LogTemp, Log, TEXT("[FLobbiesEOSGS::FindLobbies] Succeeded: User[%s]"),
+		UE_LOG(LogOnlineServices, Log, TEXT("[FLobbiesEOSGS::FindLobbies] Succeeded: User[%s]"),
 			*ToLogString(Params.LocalAccountId));
 		InAsyncOp.SetResult(FFindLobbies::Result{ LobbySearch->GetLobbyResults() });
 	})
@@ -353,14 +353,14 @@ TOnlineAsyncOpHandle<FJoinLobby> FLobbiesEOSGS::JoinLobby(FJoinLobby::Params&& I
 	// Check prerequisites.
 	if (!Services.GetAuthInterface()->IsLoggedIn(Params.LocalAccountId))
 	{
-		UE_LOG(LogTemp, Warning, TEXT("[FLobbiesEOSGS::JoinLobby] Failed: User is not logged in. User[%s]"), *ToLogString(Params.LocalAccountId));
+		UE_LOG(LogOnlineServices, Warning, TEXT("[FLobbiesEOSGS::JoinLobby] Failed: User is not logged in. User[%s]"), *ToLogString(Params.LocalAccountId));
 		Op->SetError(Errors::InvalidUser());
 		return Op->GetHandle();
 	}
 
 	if (!Params.LobbyId.IsValid())
 	{
-		UE_LOG(LogTemp, Warning, TEXT("[FLobbiesEOSGS::JoinLobby] Failed: Lobby id is invalid."));
+		UE_LOG(LogOnlineServices, Warning, TEXT("[FLobbiesEOSGS::JoinLobby] Failed: Lobby id is invalid."));
 		Op->SetError(Errors::InvalidParams());
 		return Op->GetHandle();
 	}
@@ -381,7 +381,7 @@ TOnlineAsyncOpHandle<FJoinLobby> FLobbiesEOSGS::JoinLobby(FJoinLobby::Params&& I
 	TSharedPtr<FLobbyDataEOS> LobbyData = LobbyDataRegistry->Find(Params.LobbyId);
 	if (!LobbyData)
 	{
-		UE_LOG(LogTemp, Warning, TEXT("[FLobbiesEOSGS::JoinLobby] Failed: Unable to find lobby data. LobbyId[%s]"), *ToLogString(Params.LobbyId));
+		UE_LOG(LogOnlineServices, Warning, TEXT("[FLobbiesEOSGS::JoinLobby] Failed: Unable to find lobby data. LobbyId[%s]"), *ToLogString(Params.LobbyId));
 		Op->SetError(Errors::InvalidParams());
 		return Op->GetHandle();
 	}
@@ -391,7 +391,7 @@ TOnlineAsyncOpHandle<FJoinLobby> FLobbiesEOSGS::JoinLobby(FJoinLobby::Params&& I
 	TSharedPtr<FLobbyDetailsEOS> LobbyDetails = LobbyData->GetUserLobbyDetails(Params.LocalAccountId);
 	if (!LobbyDetails)
 	{
-		UE_LOG(LogTemp, Warning, TEXT("[FLobbiesEOSGS::JoinLobby] Failed: Unable to find lobby details for user. User[%s], Lobby[%s]"),
+		UE_LOG(LogOnlineServices, Warning, TEXT("[FLobbiesEOSGS::JoinLobby] Failed: Unable to find lobby details for user. User[%s], Lobby[%s]"),
 			*ToLogString(Params.LocalAccountId), *LobbyData->GetLobbyIdString());
 		Op->SetError(Errors::InvalidParams());
 		return Op->GetHandle();
@@ -401,7 +401,7 @@ TOnlineAsyncOpHandle<FJoinLobby> FLobbiesEOSGS::JoinLobby(FJoinLobby::Params&& I
 	// Check that lobby is compatible with the running game.
 	if (!LobbyDetails->IsBucketCompatible())
 	{
-		UE_LOG(LogTemp, Warning, TEXT("[FLobbiesEOSGS::JoinLobby] Failed: Lobby is not compatible with the client version. Lobby[%s], Bucket[%s], ExpectedBucket[%s]."),
+		UE_LOG(LogOnlineServices, Warning, TEXT("[FLobbiesEOSGS::JoinLobby] Failed: Lobby is not compatible with the client version. Lobby[%s], Bucket[%s], ExpectedBucket[%s]."),
 			*LobbyData->GetLobbyIdString(), *ToLogString(LobbyDetails->GetInfo()->GetBucketId()), *ToLogString(LobbyPrerequisites->BucketId));
 		Op->SetError(Errors::IncompatibleVersion());
 		return Op->GetHandle();
@@ -432,7 +432,7 @@ TOnlineAsyncOpHandle<FJoinLobby> FLobbiesEOSGS::JoinLobby(FJoinLobby::Params&& I
 			const FJoinLobby::Params& Params = InAsyncOp.GetParams();
 			const TSharedRef<FLobbyDataEOS>& LobbyData = GetOpDataChecked<TSharedRef<FLobbyDataEOS>>(InAsyncOp, UE_ONLINE_LOBBY_EOS_KEY_NAME_LOBBY_DATA);
 
-			UE_LOG(LogTemp, Warning, TEXT("[FLobbiesEOSGS::JoinLobby] EOS_Lobby_JoinLobby Failed: User[%s], Lobby[%s], Result[%s]"),
+			UE_LOG(LogOnlineServices, Warning, TEXT("[FLobbiesEOSGS::JoinLobby] EOS_Lobby_JoinLobby Failed: User[%s], Lobby[%s], Result[%s]"),
 				*ToLogString(Params.LocalAccountId), *LobbyData->GetLobbyIdString(), *LexToString(Data->ResultCode));
 			InAsyncOp.SetError(Errors::FromEOSResult(Data->ResultCode));
 			return;
@@ -451,7 +451,7 @@ TOnlineAsyncOpHandle<FJoinLobby> FLobbiesEOSGS::JoinLobby(FJoinLobby::Params&& I
 		TOnlineResult<FLobbyClientDataPrepareClientChanges> PrepareResult = LobbyData->GetLobbyClientData()->PrepareClientChanges(MoveTemp(PrepareParams));
 		if (PrepareResult.IsError())
 		{
-			UE_LOG(LogTemp, Warning, TEXT("[FLobbiesEOSGS::JoinLobby] PrepareClientChanges Failed: User[%s], Lobby[%s], Result[%s]"),
+			UE_LOG(LogOnlineServices, Warning, TEXT("[FLobbiesEOSGS::JoinLobby] PrepareClientChanges Failed: User[%s], Lobby[%s], Result[%s]"),
 				*ToLogString(Params.LocalAccountId), *LobbyData->GetLobbyIdString(), *PrepareResult.GetErrorValue().GetLogString());
 			InAsyncOp.SetError(MoveTemp(PrepareResult.GetErrorValue()));
 			return MakeFulfilledPromise<TDefaultErrorResult<FLobbiesModifyLobbyDataImpl>>(FLobbiesModifyLobbyDataImpl::Result{}).GetFuture();
@@ -471,7 +471,7 @@ TOnlineAsyncOpHandle<FJoinLobby> FLobbiesEOSGS::JoinLobby(FJoinLobby::Params&& I
 			const FJoinLobby::Params& Params = InAsyncOp.GetParams();
 			const TSharedRef<FLobbyDataEOS>& LobbyData = GetOpDataChecked<TSharedRef<FLobbyDataEOS>>(InAsyncOp, UE_ONLINE_LOBBY_EOS_KEY_NAME_LOBBY_DATA);
 
-			UE_LOG(LogTemp, Warning, TEXT("[FLobbiesEOSGS::JoinLobby] ModifyLobbyMemberDataImpl Failed: User[%s], Lobby[%s], Result[%s]"),
+			UE_LOG(LogOnlineServices, Warning, TEXT("[FLobbiesEOSGS::JoinLobby] ModifyLobbyMemberDataImpl Failed: User[%s], Lobby[%s], Result[%s]"),
 				*ToLogString(Params.LocalAccountId), *LobbyData->GetLobbyIdString(), *Result.GetErrorValue().GetLogString());
 
 			// Failed to set attributes - leave the lobby.
@@ -494,7 +494,7 @@ TOnlineAsyncOpHandle<FJoinLobby> FLobbiesEOSGS::JoinLobby(FJoinLobby::Params&& I
 			{
 				if (Result.IsError())
 				{
-					UE_LOG(LogTemp, Warning, TEXT("[FLobbiesEOSGS::JoinLobby] LeaveLobbyImpl Failed: User[%s], Lobby[%s], Result[%s]"),
+					UE_LOG(LogOnlineServices, Warning, TEXT("[FLobbiesEOSGS::JoinLobby] LeaveLobbyImpl Failed: User[%s], Lobby[%s], Result[%s]"),
 						*ToLogString(LocalAccountId), *LobbyData->GetLobbyIdString(), *Result.GetErrorValue().GetLogString());
 				}
 
@@ -521,7 +521,7 @@ TOnlineAsyncOpHandle<FJoinLobby> FLobbiesEOSGS::JoinLobby(FJoinLobby::Params&& I
 		// Update local cache and fire events.
 		LobbyData->GetLobbyClientData()->CommitClientChanges({ &LobbyEvents });
 
-		UE_LOG(LogTemp, Log, TEXT("[FLobbiesEOSGS::JoinLobby] Succeeded: User[%s], Lobby[%s]"),
+		UE_LOG(LogOnlineServices, Log, TEXT("[FLobbiesEOSGS::JoinLobby] Succeeded: User[%s], Lobby[%s]"),
 			*ToLogString(Params.LocalAccountId), *LobbyData->GetLobbyIdString());
 		InAsyncOp.SetResult(FJoinLobby::Result{ LobbyData->GetLobbyClientData()->GetPublicDataPtr()});
 	})
@@ -538,14 +538,14 @@ TOnlineAsyncOpHandle<FLeaveLobby> FLobbiesEOSGS::LeaveLobby(FLeaveLobby::Params&
 	// Check prerequisites.
 	if (!Services.GetAuthInterface()->IsLoggedIn(Params.LocalAccountId))
 	{
-		UE_LOG(LogTemp, Warning, TEXT("[FLobbiesEOSGS::LeaveLobby] Failed: User is not logged in. User[%s]"), *ToLogString(Params.LocalAccountId));
+		UE_LOG(LogOnlineServices, Warning, TEXT("[FLobbiesEOSGS::LeaveLobby] Failed: User is not logged in. User[%s]"), *ToLogString(Params.LocalAccountId));
 		Op->SetError(Errors::InvalidUser());
 		return Op->GetHandle();
 	}
 
 	if (!Params.LobbyId.IsValid())
 	{
-		UE_LOG(LogTemp, Warning, TEXT("[FLobbiesEOSGS::LeaveLobby] Failed: Lobby id is invalid."));
+		UE_LOG(LogOnlineServices, Warning, TEXT("[FLobbiesEOSGS::LeaveLobby] Failed: Lobby id is invalid."));
 		Op->SetError(Errors::InvalidParams());
 		return Op->GetHandle();
 	}
@@ -553,7 +553,7 @@ TOnlineAsyncOpHandle<FLeaveLobby> FLobbiesEOSGS::LeaveLobby(FLeaveLobby::Params&
 	TSharedPtr<FLobbyDataEOS> LobbyData = LobbyDataRegistry->Find(Params.LobbyId);
 	if (!LobbyData)
 	{
-		UE_LOG(LogTemp, Warning, TEXT("[FLobbiesEOSGS::LeaveLobby] Failed: Unable to find lobby data. LobbyId[%s]"), *ToLogString(Params.LobbyId));
+		UE_LOG(LogOnlineServices, Warning, TEXT("[FLobbiesEOSGS::LeaveLobby] Failed: Unable to find lobby data. LobbyId[%s]"), *ToLogString(Params.LobbyId));
 		Op->SetError(Errors::InvalidParams());
 		return Op->GetHandle();
 	}
@@ -571,7 +571,7 @@ TOnlineAsyncOpHandle<FLeaveLobby> FLobbiesEOSGS::LeaveLobby(FLeaveLobby::Params&
 		TOnlineResult<FLobbyClientDataPrepareClientChanges> PrepareResult = LobbyData->GetLobbyClientData()->PrepareClientChanges(MoveTemp(PrepareParams));
 		if (PrepareResult.IsError())
 		{
-			UE_LOG(LogTemp, Warning, TEXT("[FLobbiesEOSGS::KickLobbyMember] PrepareClientChanges Failed: User[%s], Lobby[%s], Result[%s]"),
+			UE_LOG(LogOnlineServices, Warning, TEXT("[FLobbiesEOSGS::KickLobbyMember] PrepareClientChanges Failed: User[%s], Lobby[%s], Result[%s]"),
 				*ToLogString(Params.LocalAccountId), *LobbyData->GetLobbyIdString(), *PrepareResult.GetErrorValue().GetLogString());
 			InAsyncOp.SetError(MoveTemp(PrepareResult.GetErrorValue()));
 			return MakeFulfilledPromise<TDefaultErrorResult<FLobbiesLeaveLobbyImpl>>(FLobbiesLeaveLobbyImpl::Result{}).GetFuture();
@@ -589,7 +589,7 @@ TOnlineAsyncOpHandle<FLeaveLobby> FLobbiesEOSGS::LeaveLobby(FLeaveLobby::Params&
 			const FLeaveLobby::Params& Params = InAsyncOp.GetParams();
 			const TSharedRef<FLobbyDataEOS>& LobbyData = GetOpDataChecked<TSharedRef<FLobbyDataEOS>>(InAsyncOp, UE_ONLINE_LOBBY_EOS_KEY_NAME_LOBBY_DATA);
 
-			UE_LOG(LogTemp, Warning, TEXT("[FLobbiesEOSGS::LeaveLobby] LeaveLobbyImpl Failed: User[%s], Lobby[%s], Result[%s]"),
+			UE_LOG(LogOnlineServices, Warning, TEXT("[FLobbiesEOSGS::LeaveLobby] LeaveLobbyImpl Failed: User[%s], Lobby[%s], Result[%s]"),
 				*ToLogString(Params.LocalAccountId), *LobbyData->GetLobbyIdString(), *Result.GetErrorValue().GetLogString());
 
 			// When leaving a lobby consume the error and report success.
@@ -610,7 +610,7 @@ TOnlineAsyncOpHandle<FLeaveLobby> FLobbiesEOSGS::LeaveLobby(FLeaveLobby::Params&
 		// The lobby data will be cleaned up once all references are removed.
 		RemoveActiveLobby(Params.LocalAccountId, LobbyData);
 
-		UE_LOG(LogTemp, Log, TEXT("[FLobbiesEOSGS::LeaveLobby] Succeeded: User[%s], Lobby[%s]"),
+		UE_LOG(LogOnlineServices, Log, TEXT("[FLobbiesEOSGS::LeaveLobby] Succeeded: User[%s], Lobby[%s]"),
 			*ToLogString(Params.LocalAccountId), *LobbyData->GetLobbyIdString());
 		InAsyncOp.SetResult(FLeaveLobby::Result{});
 	})
@@ -627,14 +627,14 @@ TOnlineAsyncOpHandle<FInviteLobbyMember> FLobbiesEOSGS::InviteLobbyMember(FInvit
 	// Check prerequisites.
 	if (!Services.GetAuthInterface()->IsLoggedIn(Params.LocalAccountId))
 	{
-		UE_LOG(LogTemp, Warning, TEXT("[FLobbiesEOSGS::InviteLobbyMember] Failed: User is not logged in. User[%s]"), *ToLogString(Params.LocalAccountId));
+		UE_LOG(LogOnlineServices, Warning, TEXT("[FLobbiesEOSGS::InviteLobbyMember] Failed: User is not logged in. User[%s]"), *ToLogString(Params.LocalAccountId));
 		Op->SetError(Errors::InvalidUser());
 		return Op->GetHandle();
 	}
 
 	if (!Params.LobbyId.IsValid())
 	{
-		UE_LOG(LogTemp, Warning, TEXT("[FLobbiesEOSGS::InviteLobbyMember] Failed: Lobby id is invalid."));
+		UE_LOG(LogOnlineServices, Warning, TEXT("[FLobbiesEOSGS::InviteLobbyMember] Failed: Lobby id is invalid."));
 		Op->SetError(Errors::InvalidParams());
 		return Op->GetHandle();
 	}
@@ -642,7 +642,7 @@ TOnlineAsyncOpHandle<FInviteLobbyMember> FLobbiesEOSGS::InviteLobbyMember(FInvit
 	TSharedPtr<FLobbyDataEOS> LobbyData = LobbyDataRegistry->Find(Params.LobbyId);
 	if (!LobbyData)
 	{
-		UE_LOG(LogTemp, Warning, TEXT("[FLobbiesEOSGS::InviteLobbyMember] Failed: Unable to find lobby data. LobbyId[%s]"), *ToLogString(Params.LobbyId));
+		UE_LOG(LogOnlineServices, Warning, TEXT("[FLobbiesEOSGS::InviteLobbyMember] Failed: Unable to find lobby data. LobbyId[%s]"), *ToLogString(Params.LobbyId));
 		Op->SetError(Errors::InvalidParams());
 		return Op->GetHandle();
 	}
@@ -667,7 +667,7 @@ TOnlineAsyncOpHandle<FInviteLobbyMember> FLobbiesEOSGS::InviteLobbyMember(FInvit
 			const FInviteLobbyMember::Params& Params = InAsyncOp.GetParams();
 			const TSharedRef<FLobbyDataEOS>& LobbyData = GetOpDataChecked<TSharedRef<FLobbyDataEOS>>(InAsyncOp, UE_ONLINE_LOBBY_EOS_KEY_NAME_LOBBY_DATA);
 
-			UE_LOG(LogTemp, Warning, TEXT("[FLobbiesEOSGS::InviteLobbyMember] InviteLobbyMemberImpl Failed: User[%s], Lobby[%s], Result[%s]"),
+			UE_LOG(LogOnlineServices, Warning, TEXT("[FLobbiesEOSGS::InviteLobbyMember] InviteLobbyMemberImpl Failed: User[%s], Lobby[%s], Result[%s]"),
 				*ToLogString(Params.LocalAccountId), *LobbyData->GetLobbyIdString(), *Result.GetErrorValue().GetLogString());
 			InAsyncOp.SetError(MoveTemp(Result.GetErrorValue()));
 			return;
@@ -678,7 +678,7 @@ TOnlineAsyncOpHandle<FInviteLobbyMember> FLobbiesEOSGS::InviteLobbyMember(FInvit
 		const FInviteLobbyMember::Params& Params = InAsyncOp.GetParams();
 		const TSharedRef<FLobbyDataEOS>& LobbyData = GetOpDataChecked<TSharedRef<FLobbyDataEOS>>(InAsyncOp, UE_ONLINE_LOBBY_EOS_KEY_NAME_LOBBY_DATA);
 
-		UE_LOG(LogTemp, Verbose, TEXT("[FLobbiesEOSGS::InviteLobbyMember] Succeeded: User[%s], Lobby[%s]"),
+		UE_LOG(LogOnlineServices, Verbose, TEXT("[FLobbiesEOSGS::InviteLobbyMember] Succeeded: User[%s], Lobby[%s]"),
 			*ToLogString(Params.LocalAccountId), *LobbyData->GetLobbyIdString());
 		InAsyncOp.SetResult(FInviteLobbyMember::Result{});
 	})
@@ -695,14 +695,14 @@ TOnlineAsyncOpHandle<FDeclineLobbyInvitation> FLobbiesEOSGS::DeclineLobbyInvitat
 	// Check prerequisites.
 	if (!Services.GetAuthInterface()->IsLoggedIn(Params.LocalAccountId))
 	{
-		UE_LOG(LogTemp, Warning, TEXT("[FLobbiesEOSGS::DeclineLobbyInvitation] Failed: User is not logged in. User[%s]"), *ToLogString(Params.LocalAccountId));
+		UE_LOG(LogOnlineServices, Warning, TEXT("[FLobbiesEOSGS::DeclineLobbyInvitation] Failed: User is not logged in. User[%s]"), *ToLogString(Params.LocalAccountId));
 		Op->SetError(Errors::InvalidUser());
 		return Op->GetHandle();
 	}
 
 	if (!Params.LobbyId.IsValid())
 	{
-		UE_LOG(LogTemp, Warning, TEXT("[FLobbiesEOSGS::DeclineLobbyInvitation] Failed: Lobby id is invalid."));
+		UE_LOG(LogOnlineServices, Warning, TEXT("[FLobbiesEOSGS::DeclineLobbyInvitation] Failed: Lobby id is invalid."));
 		Op->SetError(Errors::InvalidParams());
 		return Op->GetHandle();
 	}
@@ -724,7 +724,7 @@ TOnlineAsyncOpHandle<FDeclineLobbyInvitation> FLobbiesEOSGS::DeclineLobbyInvitat
 			const FDeclineLobbyInvitation::Params& Params = InAsyncOp.GetParams();
 			const TSharedRef<FLobbyDataEOS>& LobbyData = GetOpDataChecked<TSharedRef<FLobbyDataEOS>>(InAsyncOp, UE_ONLINE_LOBBY_EOS_KEY_NAME_LOBBY_DATA);
 
-			UE_LOG(LogTemp, Warning, TEXT("[FLobbiesEOSGS::DeclineLobbyInvitation] DeclineLobbyInvitationImpl Failed: User[%s], Lobby[%s], Result[%s]"),
+			UE_LOG(LogOnlineServices, Warning, TEXT("[FLobbiesEOSGS::DeclineLobbyInvitation] DeclineLobbyInvitationImpl Failed: User[%s], Lobby[%s], Result[%s]"),
 				*ToLogString(Params.LocalAccountId), *LobbyData->GetLobbyIdString(), *Result.GetErrorValue().GetLogString());
 			InAsyncOp.SetError(MoveTemp(Result.GetErrorValue()));
 			return;
@@ -735,7 +735,7 @@ TOnlineAsyncOpHandle<FDeclineLobbyInvitation> FLobbiesEOSGS::DeclineLobbyInvitat
 		const FDeclineLobbyInvitation::Params& Params = InAsyncOp.GetParams();
 		const TSharedRef<FLobbyDataEOS>& LobbyData = GetOpDataChecked<TSharedRef<FLobbyDataEOS>>(InAsyncOp, UE_ONLINE_LOBBY_EOS_KEY_NAME_LOBBY_DATA);
 
-		UE_LOG(LogTemp, Verbose, TEXT("[FLobbiesEOSGS::DeclineLobbyInvitation] Succeeded: User[%s], Lobby[%s]"),
+		UE_LOG(LogOnlineServices, Verbose, TEXT("[FLobbiesEOSGS::DeclineLobbyInvitation] Succeeded: User[%s], Lobby[%s]"),
 			*ToLogString(Params.LocalAccountId), *LobbyData->GetLobbyIdString());
 		InAsyncOp.SetResult(FDeclineLobbyInvitation::Result{});
 	})
@@ -752,14 +752,14 @@ TOnlineAsyncOpHandle<FKickLobbyMember> FLobbiesEOSGS::KickLobbyMember(FKickLobby
 	// Check prerequisites.
 	if (!Services.GetAuthInterface()->IsLoggedIn(Params.LocalAccountId))
 	{
-		UE_LOG(LogTemp, Warning, TEXT("[FLobbiesEOSGS::KickLobbyMember] Failed: User is not logged in. User[%s]"), *ToLogString(Params.LocalAccountId));
+		UE_LOG(LogOnlineServices, Warning, TEXT("[FLobbiesEOSGS::KickLobbyMember] Failed: User is not logged in. User[%s]"), *ToLogString(Params.LocalAccountId));
 		Op->SetError(Errors::InvalidUser());
 		return Op->GetHandle();
 	}
 
 	if (!Params.LobbyId.IsValid())
 	{
-		UE_LOG(LogTemp, Warning, TEXT("[FLobbiesEOSGS::KickLobbyMember] Failed: Lobby id is invalid."));
+		UE_LOG(LogOnlineServices, Warning, TEXT("[FLobbiesEOSGS::KickLobbyMember] Failed: Lobby id is invalid."));
 		Op->SetError(Errors::InvalidParams());
 		return Op->GetHandle();
 	}
@@ -767,7 +767,7 @@ TOnlineAsyncOpHandle<FKickLobbyMember> FLobbiesEOSGS::KickLobbyMember(FKickLobby
 	TSharedPtr<FLobbyDataEOS> LobbyData = LobbyDataRegistry->Find(Params.LobbyId);
 	if (!LobbyData)
 	{
-		UE_LOG(LogTemp, Warning, TEXT("[FLobbiesEOSGS::KickLobbyMember] Failed: Unable to find lobby data. LobbyId[%s]"), *ToLogString(Params.LobbyId));
+		UE_LOG(LogOnlineServices, Warning, TEXT("[FLobbiesEOSGS::KickLobbyMember] Failed: Unable to find lobby data. LobbyId[%s]"), *ToLogString(Params.LobbyId));
 		Op->SetError(Errors::InvalidParams());
 		return Op->GetHandle();
 	}
@@ -785,7 +785,7 @@ TOnlineAsyncOpHandle<FKickLobbyMember> FLobbiesEOSGS::KickLobbyMember(FKickLobby
 		TOnlineResult<FLobbyClientDataPrepareClientChanges> PrepareResult = LobbyData->GetLobbyClientData()->PrepareClientChanges(MoveTemp(PrepareParams));
 		if (PrepareResult.IsError())
 		{
-			UE_LOG(LogTemp, Warning, TEXT("[FLobbiesEOSGS::KickLobbyMember] PrepareClientChanges Failed: User[%s], Lobby[%s], Result[%s]"),
+			UE_LOG(LogOnlineServices, Warning, TEXT("[FLobbiesEOSGS::KickLobbyMember] PrepareClientChanges Failed: User[%s], Lobby[%s], Result[%s]"),
 				*ToLogString(Params.LocalAccountId), *LobbyData->GetLobbyIdString(), *PrepareResult.GetErrorValue().GetLogString());
 			InAsyncOp.SetError(MoveTemp(PrepareResult.GetErrorValue()));
 			return MakeFulfilledPromise<TDefaultErrorResult<FLobbiesKickLobbyMemberImpl>>(FLobbiesKickLobbyMemberImpl::Result{}).GetFuture();
@@ -804,7 +804,7 @@ TOnlineAsyncOpHandle<FKickLobbyMember> FLobbiesEOSGS::KickLobbyMember(FKickLobby
 			const FKickLobbyMember::Params& Params = InAsyncOp.GetParams();
 			const TSharedRef<FLobbyDataEOS>& LobbyData = GetOpDataChecked<TSharedRef<FLobbyDataEOS>>(InAsyncOp, UE_ONLINE_LOBBY_EOS_KEY_NAME_LOBBY_DATA);
 
-			UE_LOG(LogTemp, Warning, TEXT("[FLobbiesEOSGS::KickLobbyMember] KickLobbyMemberImpl Failed: User[%s], Lobby[%s], Result[%s]"),
+			UE_LOG(LogOnlineServices, Warning, TEXT("[FLobbiesEOSGS::KickLobbyMember] KickLobbyMemberImpl Failed: User[%s], Lobby[%s], Result[%s]"),
 				*ToLogString(Params.LocalAccountId), *LobbyData->GetLobbyIdString(), *Result.GetErrorValue().GetLogString());
 			InAsyncOp.SetError(MoveTemp(Result.GetErrorValue()));
 		}
@@ -817,7 +817,7 @@ TOnlineAsyncOpHandle<FKickLobbyMember> FLobbiesEOSGS::KickLobbyMember(FKickLobby
 		// Update local cache and fire events.
 		LobbyData->GetLobbyClientData()->CommitClientChanges({ &LobbyEvents });
 
-		UE_LOG(LogTemp, Verbose, TEXT("[FLobbiesEOSGS::KickLobbyMember] Succeeded: User[%s], Lobby[%s]"),
+		UE_LOG(LogOnlineServices, Verbose, TEXT("[FLobbiesEOSGS::KickLobbyMember] Succeeded: User[%s], Lobby[%s]"),
 			*ToLogString(Params.LocalAccountId), *LobbyData->GetLobbyIdString());
 		InAsyncOp.SetResult(FKickLobbyMember::Result{});
 	})
@@ -834,14 +834,14 @@ TOnlineAsyncOpHandle<FPromoteLobbyMember> FLobbiesEOSGS::PromoteLobbyMember(FPro
 
 	if (!Services.GetAuthInterface()->IsLoggedIn(Params.LocalAccountId))
 	{
-		UE_LOG(LogTemp, Warning, TEXT("[FLobbiesEOSGS::FPromoteLobbyMember] Failed: User is not logged in. User[%s]"), *ToLogString(Params.LocalAccountId));
+		UE_LOG(LogOnlineServices, Warning, TEXT("[FLobbiesEOSGS::FPromoteLobbyMember] Failed: User is not logged in. User[%s]"), *ToLogString(Params.LocalAccountId));
 		Op->SetError(Errors::InvalidUser());
 		return Op->GetHandle();
 	}
 
 	if (!Params.LobbyId.IsValid())
 	{
-		UE_LOG(LogTemp, Warning, TEXT("[FLobbiesEOSGS::FPromoteLobbyMember] Failed: Lobby id is invalid."));
+		UE_LOG(LogOnlineServices, Warning, TEXT("[FLobbiesEOSGS::FPromoteLobbyMember] Failed: Lobby id is invalid."));
 		Op->SetError(Errors::InvalidParams());
 		return Op->GetHandle();
 	}
@@ -849,7 +849,7 @@ TOnlineAsyncOpHandle<FPromoteLobbyMember> FLobbiesEOSGS::PromoteLobbyMember(FPro
 	TSharedPtr<FLobbyDataEOS> LobbyData = LobbyDataRegistry->Find(Params.LobbyId);
 	if (!LobbyData)
 	{
-		UE_LOG(LogTemp, Warning, TEXT("[FLobbiesEOSGS::FPromoteLobbyMember] Failed: Unable to find lobby data. LobbyId[%s]"), *ToLogString(Params.LobbyId));
+		UE_LOG(LogOnlineServices, Warning, TEXT("[FLobbiesEOSGS::FPromoteLobbyMember] Failed: Unable to find lobby data. LobbyId[%s]"), *ToLogString(Params.LobbyId));
 		Op->SetError(Errors::InvalidParams());
 		return Op->GetHandle();
 	}
@@ -867,7 +867,7 @@ TOnlineAsyncOpHandle<FPromoteLobbyMember> FLobbiesEOSGS::PromoteLobbyMember(FPro
 		TOnlineResult<FLobbyClientDataPrepareClientChanges> PrepareResult = LobbyData->GetLobbyClientData()->PrepareClientChanges(MoveTemp(PrepareParams));
 		if (PrepareResult.IsError())
 		{
-			UE_LOG(LogTemp, Warning, TEXT("[FLobbiesEOSGS::FPromoteLobbyMember] PrepareClientChanges Failed: User[%s], Lobby[%s], Result[%s]"),
+			UE_LOG(LogOnlineServices, Warning, TEXT("[FLobbiesEOSGS::FPromoteLobbyMember] PrepareClientChanges Failed: User[%s], Lobby[%s], Result[%s]"),
 				*ToLogString(Params.LocalAccountId), *LobbyData->GetLobbyIdString(), *PrepareResult.GetErrorValue().GetLogString());
 			InAsyncOp.SetError(MoveTemp(PrepareResult.GetErrorValue()));
 			return MakeFulfilledPromise<TDefaultErrorResult<FLobbiesPromoteLobbyMemberImpl>>(FLobbiesPromoteLobbyMemberImpl::Result{}).GetFuture();
@@ -886,7 +886,7 @@ TOnlineAsyncOpHandle<FPromoteLobbyMember> FLobbiesEOSGS::PromoteLobbyMember(FPro
 			const FPromoteLobbyMember::Params& Params = InAsyncOp.GetParams();
 			const TSharedRef<FLobbyDataEOS>& LobbyData = GetOpDataChecked<TSharedRef<FLobbyDataEOS>>(InAsyncOp, UE_ONLINE_LOBBY_EOS_KEY_NAME_LOBBY_DATA);
 
-			UE_LOG(LogTemp, Warning, TEXT("[FLobbiesEOSGS::FPromoteLobbyMember] PromoteLobbyMemberImpl Failed: User[%s], Lobby[%s], Result[%s]"),
+			UE_LOG(LogOnlineServices, Warning, TEXT("[FLobbiesEOSGS::FPromoteLobbyMember] PromoteLobbyMemberImpl Failed: User[%s], Lobby[%s], Result[%s]"),
 				*ToLogString(Params.LocalAccountId), *LobbyData->GetLobbyIdString(), *Result.GetErrorValue().GetLogString());
 			InAsyncOp.SetError(MoveTemp(Result.GetErrorValue()));
 			return;
@@ -900,7 +900,7 @@ TOnlineAsyncOpHandle<FPromoteLobbyMember> FLobbiesEOSGS::PromoteLobbyMember(FPro
 		// Update local cache and fire events.
 		LobbyData->GetLobbyClientData()->CommitClientChanges({ &LobbyEvents });
 
-		UE_LOG(LogTemp, Verbose, TEXT("[FLobbiesEOSGS::FPromoteLobbyMember] Succeeded: User[%s], Lobby[%s]"),
+		UE_LOG(LogOnlineServices, Verbose, TEXT("[FLobbiesEOSGS::FPromoteLobbyMember] Succeeded: User[%s], Lobby[%s]"),
 			*ToLogString(Params.LocalAccountId), *LobbyData->GetLobbyIdString());
 		InAsyncOp.SetResult(FPromoteLobbyMember::Result{});
 	})
@@ -917,14 +917,14 @@ TOnlineAsyncOpHandle<FModifyLobbyJoinPolicy> FLobbiesEOSGS::ModifyLobbyJoinPolic
 	// Check prerequisites.
 	if (!Services.GetAuthInterface()->IsLoggedIn(Params.LocalAccountId))
 	{
-		UE_LOG(LogTemp, Warning, TEXT("[FLobbiesEOSGS::ModifyLobbyJoinPolicy] Failed: User is not logged in. User[%s]"), *ToLogString(Params.LocalAccountId));
+		UE_LOG(LogOnlineServices, Warning, TEXT("[FLobbiesEOSGS::ModifyLobbyJoinPolicy] Failed: User is not logged in. User[%s]"), *ToLogString(Params.LocalAccountId));
 		Op->SetError(Errors::InvalidUser());
 		return Op->GetHandle();
 	}
 
 	if (!Params.LobbyId.IsValid())
 	{
-		UE_LOG(LogTemp, Warning, TEXT("[FLobbiesEOSGS::ModifyLobbyJoinPolicy] Failed: Lobby id is invalid."));
+		UE_LOG(LogOnlineServices, Warning, TEXT("[FLobbiesEOSGS::ModifyLobbyJoinPolicy] Failed: Lobby id is invalid."));
 		Op->SetError(Errors::InvalidParams());
 		return Op->GetHandle();
 	}
@@ -932,7 +932,7 @@ TOnlineAsyncOpHandle<FModifyLobbyJoinPolicy> FLobbiesEOSGS::ModifyLobbyJoinPolic
 	TSharedPtr<FLobbyDataEOS> LobbyData = LobbyDataRegistry->Find(Params.LobbyId);
 	if (!LobbyData)
 	{
-		UE_LOG(LogTemp, Warning, TEXT("[FLobbiesEOSGS::ModifyLobbyJoinPolicy] Failed: Unable to find lobby data. LobbyId[%s]"), *ToLogString(Params.LobbyId));
+		UE_LOG(LogOnlineServices, Warning, TEXT("[FLobbiesEOSGS::ModifyLobbyJoinPolicy] Failed: Unable to find lobby data. LobbyId[%s]"), *ToLogString(Params.LobbyId));
 		Op->SetError(Errors::InvalidParams());
 		return Op->GetHandle();
 	}
@@ -950,7 +950,7 @@ TOnlineAsyncOpHandle<FModifyLobbyJoinPolicy> FLobbiesEOSGS::ModifyLobbyJoinPolic
 		TOnlineResult<FLobbyClientDataPrepareClientChanges> PrepareResult = LobbyData->GetLobbyClientData()->PrepareClientChanges(MoveTemp(PrepareParams));
 		if (PrepareResult.IsError())
 		{
-			UE_LOG(LogTemp, Warning, TEXT("[FLobbiesEOSGS::ModifyLobbyJoinPolicy] PrepareClientChanges Failed: User[%s], Lobby[%s], Result[%s]"),
+			UE_LOG(LogOnlineServices, Warning, TEXT("[FLobbiesEOSGS::ModifyLobbyJoinPolicy] PrepareClientChanges Failed: User[%s], Lobby[%s], Result[%s]"),
 				*ToLogString(Params.LocalAccountId), *LobbyData->GetLobbyIdString(), *PrepareResult.GetErrorValue().GetLogString());
 			InAsyncOp.SetError(MoveTemp(PrepareResult.GetErrorValue()));
 			return MakeFulfilledPromise<TDefaultErrorResult<FLobbiesModifyLobbyDataImpl>>(FLobbiesModifyLobbyDataImpl::Result{}).GetFuture();
@@ -969,7 +969,7 @@ TOnlineAsyncOpHandle<FModifyLobbyJoinPolicy> FLobbiesEOSGS::ModifyLobbyJoinPolic
 			const FModifyLobbyJoinPolicy::Params& Params = InAsyncOp.GetParams();
 			const TSharedRef<FLobbyDataEOS>& LobbyData = GetOpDataChecked<TSharedRef<FLobbyDataEOS>>(InAsyncOp, UE_ONLINE_LOBBY_EOS_KEY_NAME_LOBBY_DATA);
 
-			UE_LOG(LogTemp, Warning, TEXT("[FLobbiesEOSGS::ModifyLobbyJoinPolicy] ModifyLobbyDataImpl Failed: User[%s], Lobby[%s], Result[%s]"),
+			UE_LOG(LogOnlineServices, Warning, TEXT("[FLobbiesEOSGS::ModifyLobbyJoinPolicy] ModifyLobbyDataImpl Failed: User[%s], Lobby[%s], Result[%s]"),
 				*ToLogString(Params.LocalAccountId), *LobbyData->GetLobbyIdString(), *Result.GetErrorValue().GetLogString());
 			InAsyncOp.SetError(MoveTemp(Result.GetErrorValue()));
 			return;
@@ -983,7 +983,7 @@ TOnlineAsyncOpHandle<FModifyLobbyJoinPolicy> FLobbiesEOSGS::ModifyLobbyJoinPolic
 		// Update local cache and fire events.
 		LobbyData->GetLobbyClientData()->CommitClientChanges({ &LobbyEvents });
 
-		UE_LOG(LogTemp, Verbose, TEXT("[FLobbiesEOSGS::ModifyLobbyJoinPolicy] Succeeded: User[%s], Lobby[%s]"),
+		UE_LOG(LogOnlineServices, Verbose, TEXT("[FLobbiesEOSGS::ModifyLobbyJoinPolicy] Succeeded: User[%s], Lobby[%s]"),
 			*ToLogString(Params.LocalAccountId), *LobbyData->GetLobbyIdString());
 		InAsyncOp.SetResult(FModifyLobbyJoinPolicy::Result{});
 	})
@@ -1000,14 +1000,14 @@ TOnlineAsyncOpHandle<FModifyLobbyAttributes> FLobbiesEOSGS::ModifyLobbyAttribute
 	// Check prerequisites.
 	if (!Services.GetAuthInterface()->IsLoggedIn(Params.LocalAccountId))
 	{
-		UE_LOG(LogTemp, Warning, TEXT("[FLobbiesEOSGS::ModifyLobbyAttributes] Failed: User is not logged in. User[%s]"), *ToLogString(Params.LocalAccountId));
+		UE_LOG(LogOnlineServices, Warning, TEXT("[FLobbiesEOSGS::ModifyLobbyAttributes] Failed: User is not logged in. User[%s]"), *ToLogString(Params.LocalAccountId));
 		Op->SetError(Errors::InvalidUser());
 		return Op->GetHandle();
 	}
 
 	if (!Params.LobbyId.IsValid())
 	{
-		UE_LOG(LogTemp, Warning, TEXT("[FLobbiesEOSGS::ModifyLobbyAttributes] Failed: Lobby id is invalid."));
+		UE_LOG(LogOnlineServices, Warning, TEXT("[FLobbiesEOSGS::ModifyLobbyAttributes] Failed: Lobby id is invalid."));
 		Op->SetError(Errors::InvalidParams());
 		return Op->GetHandle();
 	}
@@ -1015,7 +1015,7 @@ TOnlineAsyncOpHandle<FModifyLobbyAttributes> FLobbiesEOSGS::ModifyLobbyAttribute
 	TSharedPtr<FLobbyDataEOS> LobbyData = LobbyDataRegistry->Find(Params.LobbyId);
 	if (!LobbyData)
 	{
-		UE_LOG(LogTemp, Warning, TEXT("[FLobbiesEOSGS::ModifyLobbyAttributes] Failed: Unable to find lobby data. LobbyId[%s]"), *ToLogString(Params.LobbyId));
+		UE_LOG(LogOnlineServices, Warning, TEXT("[FLobbiesEOSGS::ModifyLobbyAttributes] Failed: Unable to find lobby data. LobbyId[%s]"), *ToLogString(Params.LobbyId));
 		Op->SetError(Errors::InvalidParams());
 		return Op->GetHandle();
 	}
@@ -1033,7 +1033,7 @@ TOnlineAsyncOpHandle<FModifyLobbyAttributes> FLobbiesEOSGS::ModifyLobbyAttribute
 		TOnlineResult<FLobbyClientDataPrepareClientChanges> PrepareResult = LobbyData->GetLobbyClientData()->PrepareClientChanges(MoveTemp(PrepareParams));
 		if (PrepareResult.IsError())
 		{
-			UE_LOG(LogTemp, Warning, TEXT("[FLobbiesEOSGS::ModifyLobbyAttributes] PrepareClientChanges Failed: User[%s], Lobby[%s], Result[%s]"),
+			UE_LOG(LogOnlineServices, Warning, TEXT("[FLobbiesEOSGS::ModifyLobbyAttributes] PrepareClientChanges Failed: User[%s], Lobby[%s], Result[%s]"),
 				*ToLogString(Params.LocalAccountId), *LobbyData->GetLobbyIdString(), *PrepareResult.GetErrorValue().GetLogString());
 			InAsyncOp.SetError(MoveTemp(PrepareResult.GetErrorValue()));
 			return MakeFulfilledPromise<TDefaultErrorResult<FLobbiesModifyLobbyDataImpl>>(FLobbiesModifyLobbyDataImpl::Result{}).GetFuture();
@@ -1052,7 +1052,7 @@ TOnlineAsyncOpHandle<FModifyLobbyAttributes> FLobbiesEOSGS::ModifyLobbyAttribute
 			const FModifyLobbyAttributes::Params& Params = InAsyncOp.GetParams();
 			const TSharedRef<FLobbyDataEOS>& LobbyData = GetOpDataChecked<TSharedRef<FLobbyDataEOS>>(InAsyncOp, UE_ONLINE_LOBBY_EOS_KEY_NAME_LOBBY_DATA);
 
-			UE_LOG(LogTemp, Warning, TEXT("[FLobbiesEOSGS::ModifyLobbyAttributes] ModifyLobbyDataImpl Failed: User[%s], Lobby[%s], Result[%s]"),
+			UE_LOG(LogOnlineServices, Warning, TEXT("[FLobbiesEOSGS::ModifyLobbyAttributes] ModifyLobbyDataImpl Failed: User[%s], Lobby[%s], Result[%s]"),
 				*ToLogString(Params.LocalAccountId), *LobbyData->GetLobbyIdString(), *Result.GetErrorValue().GetLogString());
 			InAsyncOp.SetError(MoveTemp(Result.GetErrorValue()));
 			return;
@@ -1066,7 +1066,7 @@ TOnlineAsyncOpHandle<FModifyLobbyAttributes> FLobbiesEOSGS::ModifyLobbyAttribute
 		// Update local cache and fire events.
 		LobbyData->GetLobbyClientData()->CommitClientChanges({ &LobbyEvents });
 
-		UE_LOG(LogTemp, Verbose, TEXT("[FLobbiesEOSGS::ModifyLobbyAttributes] Succeeded: User[%s], Lobby[%s]"),
+		UE_LOG(LogOnlineServices, Verbose, TEXT("[FLobbiesEOSGS::ModifyLobbyAttributes] Succeeded: User[%s], Lobby[%s]"),
 			*ToLogString(Params.LocalAccountId), *LobbyData->GetLobbyIdString());
 		InAsyncOp.SetResult(FModifyLobbyAttributes::Result{});
 	})
@@ -1083,14 +1083,14 @@ TOnlineAsyncOpHandle<FModifyLobbyMemberAttributes> FLobbiesEOSGS::ModifyLobbyMem
 	// Check prerequisites.
 	if (!Services.GetAuthInterface()->IsLoggedIn(Params.LocalAccountId))
 	{
-		UE_LOG(LogTemp, Warning, TEXT("[FLobbiesEOSGS::ModifyLobbyMemberAttributes] Failed: User is not logged in. User[%s]"), *ToLogString(Params.LocalAccountId));
+		UE_LOG(LogOnlineServices, Warning, TEXT("[FLobbiesEOSGS::ModifyLobbyMemberAttributes] Failed: User is not logged in. User[%s]"), *ToLogString(Params.LocalAccountId));
 		Op->SetError(Errors::InvalidUser());
 		return Op->GetHandle();
 	}
 
 	if (!Params.LobbyId.IsValid())
 	{
-		UE_LOG(LogTemp, Warning, TEXT("[FLobbiesEOSGS::ModifyLobbyMemberAttributes] Failed: Lobby id is invalid."));
+		UE_LOG(LogOnlineServices, Warning, TEXT("[FLobbiesEOSGS::ModifyLobbyMemberAttributes] Failed: Lobby id is invalid."));
 		Op->SetError(Errors::InvalidParams());
 		return Op->GetHandle();
 	}
@@ -1098,7 +1098,7 @@ TOnlineAsyncOpHandle<FModifyLobbyMemberAttributes> FLobbiesEOSGS::ModifyLobbyMem
 	TSharedPtr<FLobbyDataEOS> LobbyData = LobbyDataRegistry->Find(Params.LobbyId);
 	if (!LobbyData)
 	{
-		UE_LOG(LogTemp, Warning, TEXT("[FLobbiesEOSGS::ModifyLobbyMemberAttributes] Failed: Unable to find lobby data. LobbyId[%s]"), *ToLogString(Params.LobbyId));
+		UE_LOG(LogOnlineServices, Warning, TEXT("[FLobbiesEOSGS::ModifyLobbyMemberAttributes] Failed: Unable to find lobby data. LobbyId[%s]"), *ToLogString(Params.LobbyId));
 		Op->SetError(Errors::InvalidParams());
 		return Op->GetHandle();
 	}
@@ -1116,7 +1116,7 @@ TOnlineAsyncOpHandle<FModifyLobbyMemberAttributes> FLobbiesEOSGS::ModifyLobbyMem
 		TOnlineResult<FLobbyClientDataPrepareClientChanges> PrepareResult = LobbyData->GetLobbyClientData()->PrepareClientChanges(MoveTemp(PrepareParams));
 		if (PrepareResult.IsError())
 		{
-			UE_LOG(LogTemp, Warning, TEXT("[FLobbiesEOSGS::ModifyLobbyMemberAttributes] PrepareClientChanges Failed: User[%s], Lobby[%s], Result[%s]"),
+			UE_LOG(LogOnlineServices, Warning, TEXT("[FLobbiesEOSGS::ModifyLobbyMemberAttributes] PrepareClientChanges Failed: User[%s], Lobby[%s], Result[%s]"),
 				*ToLogString(Params.LocalAccountId), *LobbyData->GetLobbyIdString(), *PrepareResult.GetErrorValue().GetLogString());
 			InAsyncOp.SetError(MoveTemp(PrepareResult.GetErrorValue()));
 			return MakeFulfilledPromise<TDefaultErrorResult<FLobbiesModifyLobbyDataImpl>>(FLobbiesModifyLobbyDataImpl::Result{}).GetFuture();
@@ -1135,7 +1135,7 @@ TOnlineAsyncOpHandle<FModifyLobbyMemberAttributes> FLobbiesEOSGS::ModifyLobbyMem
 			const FModifyLobbyMemberAttributes::Params& Params = InAsyncOp.GetParams();
 			const TSharedRef<FLobbyDataEOS>& LobbyData = GetOpDataChecked<TSharedRef<FLobbyDataEOS>>(InAsyncOp, UE_ONLINE_LOBBY_EOS_KEY_NAME_LOBBY_DATA);
 
-			UE_LOG(LogTemp, Warning, TEXT("[FLobbiesEOSGS::ModifyLobbyMemberAttributes] ModifyLobbyMemberDataImpl Failed: User[%s], Lobby[%s], Result[%s]"),
+			UE_LOG(LogOnlineServices, Warning, TEXT("[FLobbiesEOSGS::ModifyLobbyMemberAttributes] ModifyLobbyMemberDataImpl Failed: User[%s], Lobby[%s], Result[%s]"),
 				*ToLogString(Params.LocalAccountId), *LobbyData->GetLobbyIdString(), *Result.GetErrorValue().GetLogString());
 			InAsyncOp.SetError(MoveTemp(Result.GetErrorValue()));
 			return;
@@ -1149,7 +1149,7 @@ TOnlineAsyncOpHandle<FModifyLobbyMemberAttributes> FLobbiesEOSGS::ModifyLobbyMem
 		// Update local cache and fire events.
 		LobbyData->GetLobbyClientData()->CommitClientChanges({ &LobbyEvents });
 
-		UE_LOG(LogTemp, Verbose, TEXT("[FLobbiesEOSGS::ModifyLobbyMemberAttributes] Succeeded: User[%s], Lobby[%s]"),
+		UE_LOG(LogOnlineServices, Verbose, TEXT("[FLobbiesEOSGS::ModifyLobbyMemberAttributes] Succeeded: User[%s], Lobby[%s]"),
 			*ToLogString(Params.LocalAccountId), *LobbyData->GetLobbyIdString());
 		InAsyncOp.SetResult(FModifyLobbyMemberAttributes::Result{});
 	})
@@ -1189,12 +1189,12 @@ void FLobbiesEOSGS::HandleLobbyUpdated(const EOS_Lobby_LobbyUpdateReceivedCallba
 			if (Result.IsError())
 			{
 				// Todo: handle failure to update lobby from snapshot.
-				UE_LOG(LogTemp, Warning, TEXT("[FLobbiesEOSGS::HandleLobbyUpdated] ProcessLobbyNotificationImplOp Failed: Lobby[%s], Result[%s]"),
+				UE_LOG(LogOnlineServices, Warning, TEXT("[FLobbiesEOSGS::HandleLobbyUpdated] ProcessLobbyNotificationImplOp Failed: Lobby[%s], Result[%s]"),
 					*LobbyData->GetLobbyIdString(), *Result.GetErrorValue().GetLogString());
 			}
 			else
 			{
-				UE_LOG(LogTemp, Verbose, TEXT("[FLobbiesEOSGS::HandleLobbyUpdated] Succeeded: Lobby[%s]"),
+				UE_LOG(LogOnlineServices, Verbose, TEXT("[FLobbiesEOSGS::HandleLobbyUpdated] Succeeded: Lobby[%s]"),
 					*LobbyData->GetLobbyIdString());
 			}
 		});
@@ -1215,12 +1215,12 @@ void FLobbiesEOSGS::HandleLobbyMemberUpdated(const EOS_Lobby_LobbyMemberUpdateRe
 			if (Result.IsError())
 			{
 				// Todo: handle failure to update lobby from snapshot.
-				UE_LOG(LogTemp, Warning, TEXT("[FLobbiesEOSGS::HandleLobbyMemberUpdated] ProcessLobbyNotificationImplOp Failed: Lobby[%s], Result[%s]"),
+				UE_LOG(LogOnlineServices, Warning, TEXT("[FLobbiesEOSGS::HandleLobbyMemberUpdated] ProcessLobbyNotificationImplOp Failed: Lobby[%s], Result[%s]"),
 					*LobbyData->GetLobbyIdString(), *Result.GetErrorValue().GetLogString());
 			}
 			else
 			{
-				UE_LOG(LogTemp, Verbose, TEXT("[FLobbiesEOSGS::HandleLobbyMemberUpdated] Succeeded: Lobby[%s]"),
+				UE_LOG(LogOnlineServices, Verbose, TEXT("[FLobbiesEOSGS::HandleLobbyMemberUpdated] Succeeded: Lobby[%s]"),
 					*LobbyData->GetLobbyIdString());
 			}
 		});
@@ -1269,12 +1269,12 @@ void FLobbiesEOSGS::HandleLobbyMemberStatusReceived(const EOS_Lobby_LobbyMemberS
 			if (Result.IsError())
 			{
 				// Todo: handle failure to update lobby from snapshot.
-				UE_LOG(LogTemp, Warning, TEXT("[FLobbiesEOSGS::HandleLobbyMemberStatusReceived] ProcessLobbyNotificationImplOp Failed: Lobby[%s], Result[%s]"),
+				UE_LOG(LogOnlineServices, Warning, TEXT("[FLobbiesEOSGS::HandleLobbyMemberStatusReceived] ProcessLobbyNotificationImplOp Failed: Lobby[%s], Result[%s]"),
 					*LobbyData->GetLobbyIdString(), *Result.GetErrorValue().GetLogString());
 			}
 			else
 			{
-				UE_LOG(LogTemp, Verbose, TEXT("[FLobbiesEOSGS::HandleLobbyMemberStatusReceived] Succeeded: Lobby[%s]"),
+				UE_LOG(LogOnlineServices, Verbose, TEXT("[FLobbiesEOSGS::HandleLobbyMemberStatusReceived] Succeeded: Lobby[%s]"),
 					*LobbyData->GetLobbyIdString());
 			}
 		});
@@ -1293,12 +1293,12 @@ void FLobbiesEOSGS::HandleLobbyInviteReceived(const EOS_Lobby_LobbyInviteReceive
 			if (Result.IsError())
 			{
 				// Todo: Log / queue a manual fetch of invitations.
-				UE_LOG(LogTemp, Warning, TEXT("[FLobbiesEOSGS::HandleLobbyInviteReceived] FLobbyInviteDataEOS::CreateFromInviteId Failed: Result[%s]"),
+				UE_LOG(LogOnlineServices, Warning, TEXT("[FLobbiesEOSGS::HandleLobbyInviteReceived] FLobbyInviteDataEOS::CreateFromInviteId Failed: Result[%s]"),
 					*Result.GetErrorValue().GetLogString());
 			}
 			else
 			{
-				UE_LOG(LogTemp, Verbose, TEXT("[FLobbiesEOSGS::HandleLobbyMemberStatusReceived] Succeeded: InviteId[%s], Lobby[%s], Receiver[%s], Sender[%s]"),
+				UE_LOG(LogOnlineServices, Verbose, TEXT("[FLobbiesEOSGS::HandleLobbyMemberStatusReceived] Succeeded: InviteId[%s], Lobby[%s], Receiver[%s], Sender[%s]"),
 					*Result.GetOkValue()->GetInviteId(),
 					*Result.GetOkValue()->GetLobbyData()->GetLobbyIdString(),
 					*ToLogString(Result.GetOkValue()->GetReceiver()),
@@ -1321,12 +1321,12 @@ void FLobbiesEOSGS::HandleLobbyInviteAccepted(const EOS_Lobby_LobbyInviteAccepte
 			if (Result.IsError())
 			{
 				// Todo: Log / queue a manual fetch of invitations.
-				UE_LOG(LogTemp, Warning, TEXT("[FLobbiesEOSGS::HandleLobbyInviteAccepted] FLobbyInviteDataEOS::CreateFromInviteId Failed: Result[%s]"),
+				UE_LOG(LogOnlineServices, Warning, TEXT("[FLobbiesEOSGS::HandleLobbyInviteAccepted] FLobbyInviteDataEOS::CreateFromInviteId Failed: Result[%s]"),
 					*Result.GetErrorValue().GetLogString());
 			}
 			else
 			{
-				UE_LOG(LogTemp, Verbose, TEXT("[FLobbiesEOSGS::HandleLobbyInviteAccepted] Succeeded: InviteId[%s], Lobby[%s], Receiver[%s], Sender[%s]"),
+				UE_LOG(LogOnlineServices, Verbose, TEXT("[FLobbiesEOSGS::HandleLobbyInviteAccepted] Succeeded: InviteId[%s], Lobby[%s], Receiver[%s], Sender[%s]"),
 					*Result.GetOkValue()->GetInviteId(),
 					*Result.GetOkValue()->GetLobbyData()->GetLobbyIdString(),
 					*ToLogString(Result.GetOkValue()->GetReceiver()),
@@ -1355,7 +1355,7 @@ void FLobbiesEOSGS::HandleJoinLobbyAccepted(const EOS_Lobby_JoinLobbyAcceptedCal
 		if (LobbyDetailsResult.IsError())
 		{
 			// Todo: Log / queue a manual fetch of invitations.
-			UE_LOG(LogTemp, Warning, TEXT("[FLobbiesEOSGS::HandleJoinLobbyAccepted] FLobbyDetailsEOS::CreateFromUiEventId Failed: User[%s], Result[%s]"),
+			UE_LOG(LogOnlineServices, Warning, TEXT("[FLobbiesEOSGS::HandleJoinLobbyAccepted] FLobbyDetailsEOS::CreateFromUiEventId Failed: User[%s], Result[%s]"),
 				*ToLogString(LocalAccountId), *LobbyDetailsResult.GetErrorValue().GetLogString());
 			return;
 		}
@@ -1365,12 +1365,12 @@ void FLobbiesEOSGS::HandleJoinLobbyAccepted(const EOS_Lobby_JoinLobbyAcceptedCal
 		{
 			if (Result.IsError())
 			{
-				UE_LOG(LogTemp, Warning, TEXT("[FLobbiesEOSGS::HandleJoinLobbyAccepted] FLobbyDataRegistryEOS::FindOrCreateFromLobbyDetails Failed: User[%s], Result[%s]"),
+				UE_LOG(LogOnlineServices, Warning, TEXT("[FLobbiesEOSGS::HandleJoinLobbyAccepted] FLobbyDataRegistryEOS::FindOrCreateFromLobbyDetails Failed: User[%s], Result[%s]"),
 					*ToLogString(LocalAccountId), *Result.GetErrorValue().GetLogString());
 			}
 			else
 			{
-				UE_LOG(LogTemp, Verbose, TEXT("[FLobbiesEOSGS::HandleJoinLobbyAccepted] Succeeded: User[%s], Lobby[%s]"),
+				UE_LOG(LogOnlineServices, Verbose, TEXT("[FLobbiesEOSGS::HandleJoinLobbyAccepted] Succeeded: User[%s], Lobby[%s]"),
 					*ToLogString(LocalAccountId), *Result.GetOkValue()->GetLobbyIdString());
 
 				LobbyEvents.OnUILobbyJoinRequested.Broadcast(
@@ -1384,7 +1384,7 @@ void FLobbiesEOSGS::HandleJoinLobbyAccepted(const EOS_Lobby_JoinLobbyAcceptedCal
 	}
 	else
 	{
-		UE_LOG(LogTemp, Warning, TEXT("[FLobbiesEOSGS::HandleJoinLobbyAccepted] Failed: Invalid local user."));
+		UE_LOG(LogOnlineServices, Warning, TEXT("[FLobbiesEOSGS::HandleJoinLobbyAccepted] Failed: Invalid local user."));
 	}
 }
 
@@ -1571,13 +1571,13 @@ TFuture<TDefaultErrorResult<FLobbiesLeaveLobbyImpl>> FLobbiesEOSGS::LeaveLobbyIm
 	// Check prerequisites.
 	if (!Services.GetAuthInterface()->IsLoggedIn(Params.LocalAccountId))
 	{
-		UE_LOG(LogTemp, Warning, TEXT("[FLobbiesEOSGS::LeaveLobbyImpl] Failed: User is not logged in. User[%s]"), *ToLogString(Params.LocalAccountId));
+		UE_LOG(LogOnlineServices, Warning, TEXT("[FLobbiesEOSGS::LeaveLobbyImpl] Failed: User is not logged in. User[%s]"), *ToLogString(Params.LocalAccountId));
 		return MakeFulfilledPromise<TDefaultErrorResult<FLobbiesLeaveLobbyImpl>>(Errors::InvalidUser()).GetFuture();
 	}
 
 	if (!Params.LobbyData.IsValid())
 	{
-		UE_LOG(LogTemp, Warning, TEXT("[FLobbiesEOSGS::LeaveLobbyImpl] Failed: No lobby data provided."));
+		UE_LOG(LogOnlineServices, Warning, TEXT("[FLobbiesEOSGS::LeaveLobbyImpl] Failed: No lobby data provided."));
 		return MakeFulfilledPromise<TDefaultErrorResult<FLobbiesLeaveLobbyImpl>>(Errors::InvalidParams()).GetFuture();
 	}
 
@@ -1600,13 +1600,13 @@ TFuture<TDefaultErrorResult<FLobbiesLeaveLobbyImpl>> FLobbiesEOSGS::LeaveLobbyIm
 	{
 		if (Result->ResultCode != EOS_EResult::EOS_Success)
 		{
-			UE_LOG(LogTemp, Warning, TEXT("[FLobbiesEOSGS::DestroyLobbyImpl] EOS_Lobby_DestroyLobby Failed: User[%s], Lobby[%s], Result[%s]"),
+			UE_LOG(LogOnlineServices, Warning, TEXT("[FLobbiesEOSGS::DestroyLobbyImpl] EOS_Lobby_DestroyLobby Failed: User[%s], Lobby[%s], Result[%s]"),
 				*ToLogString(LocalAccountId), *LobbyData->GetLobbyIdString(), *LexToString(Result->ResultCode));
 			Promise.EmplaceValue(Errors::FromEOSResult(Result->ResultCode));
 			return;
 		}
 
-		UE_LOG(LogTemp, Verbose, TEXT("[FLobbiesEOSGS::DestroyLobbyImpl] Succeeded: User[%s], Lobby[%s]"),
+		UE_LOG(LogOnlineServices, Verbose, TEXT("[FLobbiesEOSGS::DestroyLobbyImpl] Succeeded: User[%s], Lobby[%s]"),
 			*ToLogString(LocalAccountId), *LobbyData->GetLobbyIdString());
 		Promise.EmplaceValue(FLobbiesLeaveLobbyImpl::Result{});
 	});
@@ -1619,13 +1619,13 @@ TFuture<TDefaultErrorResult<FLobbiesDestroyLobbyImpl>> FLobbiesEOSGS::DestroyLob
 	// Check prerequisites.
 	if (!Services.GetAuthInterface()->IsLoggedIn(Params.LocalAccountId))
 	{
-		UE_LOG(LogTemp, Warning, TEXT("[FLobbiesEOSGS::DestroyLobbyImpl] Failed: User is not logged in. User[%s]"), *ToLogString(Params.LocalAccountId));
+		UE_LOG(LogOnlineServices, Warning, TEXT("[FLobbiesEOSGS::DestroyLobbyImpl] Failed: User is not logged in. User[%s]"), *ToLogString(Params.LocalAccountId));
 		return MakeFulfilledPromise<TDefaultErrorResult<FLobbiesDestroyLobbyImpl>>(Errors::InvalidUser()).GetFuture();
 	}
 
 	if (Params.LobbyIdString.IsEmpty())
 	{
-		UE_LOG(LogTemp, Warning, TEXT("[FLobbiesEOSGS::DestroyLobbyImpl] Failed: No lobby id provided."));
+		UE_LOG(LogOnlineServices, Warning, TEXT("[FLobbiesEOSGS::DestroyLobbyImpl] Failed: No lobby id provided."));
 		return MakeFulfilledPromise<TDefaultErrorResult<FLobbiesDestroyLobbyImpl>>(Errors::InvalidParams()).GetFuture();
 	}
 
@@ -1649,13 +1649,13 @@ TFuture<TDefaultErrorResult<FLobbiesDestroyLobbyImpl>> FLobbiesEOSGS::DestroyLob
 	{
 		if (Result->ResultCode != EOS_EResult::EOS_Success)
 		{
-			UE_LOG(LogTemp, Warning, TEXT("[FLobbiesEOSGS::DestroyLobbyImpl] EOS_Lobby_DestroyLobby Failed: User[%s], Lobby[%s], Result[%s]"),
+			UE_LOG(LogOnlineServices, Warning, TEXT("[FLobbiesEOSGS::DestroyLobbyImpl] EOS_Lobby_DestroyLobby Failed: User[%s], Lobby[%s], Result[%s]"),
 				*ToLogString(LocalAccountId), *LobbyIdString, *LexToString(Result->ResultCode));
 			Promise.EmplaceValue(Errors::FromEOSResult(Result->ResultCode));
 			return;
 		}
 
-		UE_LOG(LogTemp, Verbose, TEXT("[FLobbiesEOSGS::DestroyLobbyImpl] Succeeded: User[%s], Lobby[%s]"),
+		UE_LOG(LogOnlineServices, Verbose, TEXT("[FLobbiesEOSGS::DestroyLobbyImpl] Succeeded: User[%s], Lobby[%s]"),
 			*ToLogString(LocalAccountId), *LobbyIdString);
 		Promise.EmplaceValue(FLobbiesDestroyLobbyImpl::Result{});
 	});
@@ -1668,13 +1668,13 @@ TFuture<TDefaultErrorResult<FLobbiesInviteLobbyMemberImpl>> FLobbiesEOSGS::Invit
 	// Check prerequisites.
 	if (!Services.GetAuthInterface()->IsLoggedIn(Params.LocalAccountId))
 	{
-		UE_LOG(LogTemp, Warning, TEXT("[FLobbiesEOSGS::InviteLobbyMemberImpl] Failed: User is not logged in. User[%s]"), *ToLogString(Params.LocalAccountId));
+		UE_LOG(LogOnlineServices, Warning, TEXT("[FLobbiesEOSGS::InviteLobbyMemberImpl] Failed: User is not logged in. User[%s]"), *ToLogString(Params.LocalAccountId));
 		return MakeFulfilledPromise<TDefaultErrorResult<FLobbiesInviteLobbyMemberImpl>>(Errors::InvalidUser()).GetFuture();
 	}
 
 	if (!Params.LobbyData.IsValid())
 	{
-		UE_LOG(LogTemp, Warning, TEXT("[FLobbiesEOSGS::InviteLobbyMemberImpl] Failed: No lobby data provided."));
+		UE_LOG(LogOnlineServices, Warning, TEXT("[FLobbiesEOSGS::InviteLobbyMemberImpl] Failed: No lobby data provided."));
 		return MakeFulfilledPromise<TDefaultErrorResult<FLobbiesInviteLobbyMemberImpl>>(Errors::InvalidParams()).GetFuture();
 	}
 
@@ -1699,13 +1699,13 @@ TFuture<TDefaultErrorResult<FLobbiesInviteLobbyMemberImpl>> FLobbiesEOSGS::Invit
 	{
 		if (Result->ResultCode != EOS_EResult::EOS_Success)
 		{
-			UE_LOG(LogTemp, Warning, TEXT("[FLobbiesEOSGS::InviteLobbyMemberImpl] EOS_Lobby_SendInvite Failed: User[%s], TargetUser[%s], Lobby[%s], Result[%s]"),
+			UE_LOG(LogOnlineServices, Warning, TEXT("[FLobbiesEOSGS::InviteLobbyMemberImpl] EOS_Lobby_SendInvite Failed: User[%s], TargetUser[%s], Lobby[%s], Result[%s]"),
 				*ToLogString(LocalAccountId), *ToLogString(TargetAccountId), *LobbyData->GetLobbyIdString(), *LexToString(Result->ResultCode));
 			Promise.EmplaceValue(Errors::FromEOSResult(Result->ResultCode));
 			return;
 		}
 
-		UE_LOG(LogTemp, Verbose, TEXT("[FLobbiesEOSGS::InviteLobbyMemberImpl] Succeeded: User[%s], TargetUser[%s], Lobby[%s]"),
+		UE_LOG(LogOnlineServices, Verbose, TEXT("[FLobbiesEOSGS::InviteLobbyMemberImpl] Succeeded: User[%s], TargetUser[%s], Lobby[%s]"),
 			*ToLogString(LocalAccountId), *ToLogString(TargetAccountId), *LobbyData->GetLobbyIdString());
 		Promise.EmplaceValue(FLobbiesInviteLobbyMemberImpl::Result{});
 	});
@@ -1718,13 +1718,13 @@ TFuture<TDefaultErrorResult<FLobbiesDeclineLobbyInvitationImpl>> FLobbiesEOSGS::
 	// Check prerequisites.
 	if (!Services.GetAuthInterface()->IsLoggedIn(Params.LocalAccountId))
 	{
-		UE_LOG(LogTemp, Warning, TEXT("[FLobbiesEOSGS::DeclineLobbyInvitationImpl] Failed: User is not logged in. User[%s]"), *ToLogString(Params.LocalAccountId));
+		UE_LOG(LogOnlineServices, Warning, TEXT("[FLobbiesEOSGS::DeclineLobbyInvitationImpl] Failed: User is not logged in. User[%s]"), *ToLogString(Params.LocalAccountId));
 		return MakeFulfilledPromise<TDefaultErrorResult<FLobbiesDeclineLobbyInvitationImpl>>(Errors::InvalidUser()).GetFuture();
 	}
 
 	if (!Params.LobbyId.IsValid())
 	{
-		UE_LOG(LogTemp, Warning, TEXT("[FLobbiesEOSGS::DeclineLobbyInvitationImpl] Failed: Lobby id is invalid."));
+		UE_LOG(LogOnlineServices, Warning, TEXT("[FLobbiesEOSGS::DeclineLobbyInvitationImpl] Failed: Lobby id is invalid."));
 		return MakeFulfilledPromise<TDefaultErrorResult<FLobbiesDeclineLobbyInvitationImpl>>(Errors::InvalidParams()).GetFuture();
 	}
 
@@ -1732,7 +1732,7 @@ TFuture<TDefaultErrorResult<FLobbiesDeclineLobbyInvitationImpl>> FLobbiesEOSGS::
 	TSharedPtr<FLobbyInviteDataEOS> InviteData = GetActiveInvite(Params.LocalAccountId, Params.LobbyId);
 	if (!InviteData)
 	{
-		UE_LOG(LogTemp, Warning, TEXT("[FLobbiesEOSGS::DeclineLobbyInvitationImpl] Failed: Unable to find invitation for lobby. User[%s], Lobby[%s]"),
+		UE_LOG(LogOnlineServices, Warning, TEXT("[FLobbiesEOSGS::DeclineLobbyInvitationImpl] Failed: Unable to find invitation for lobby. User[%s], Lobby[%s]"),
 			*ToLogString(Params.LocalAccountId), *ToLogString(Params.LobbyId));
 		return MakeFulfilledPromise<TDefaultErrorResult<FLobbiesDeclineLobbyInvitationImpl>>(Errors::InvalidParams()).GetFuture();
 	}
@@ -1760,13 +1760,13 @@ TFuture<TDefaultErrorResult<FLobbiesDeclineLobbyInvitationImpl>> FLobbiesEOSGS::
 
 		if (Result->ResultCode != EOS_EResult::EOS_Success)
 		{
-			UE_LOG(LogTemp, Warning, TEXT("[FLobbiesEOSGS::DeclineLobbyInvitationImpl] EOS_Lobby_RejectInvite Failed: User[%s], Invitation[%s], Lobby[%s], Result[%s]"),
+			UE_LOG(LogOnlineServices, Warning, TEXT("[FLobbiesEOSGS::DeclineLobbyInvitationImpl] EOS_Lobby_RejectInvite Failed: User[%s], Invitation[%s], Lobby[%s], Result[%s]"),
 				*ToLogString(LocalAccountId), *InviteData->GetInviteId(), *InviteData->GetLobbyData()->GetLobbyIdString(), *LexToString(Result->ResultCode));
 			Promise.EmplaceValue(Errors::FromEOSResult(Result->ResultCode));
 			return;
 		}
 
-		UE_LOG(LogTemp, Verbose, TEXT("[FLobbiesEOSGS::DeclineLobbyInvitationImpl] Succeeded: User[%s], Invitation[%s], Lobby[%s]"),
+		UE_LOG(LogOnlineServices, Verbose, TEXT("[FLobbiesEOSGS::DeclineLobbyInvitationImpl] Succeeded: User[%s], Invitation[%s], Lobby[%s]"),
 			*ToLogString(LocalAccountId), *InviteData->GetInviteId(), *InviteData->GetLobbyData()->GetLobbyIdString());
 		Promise.EmplaceValue(FLobbiesDeclineLobbyInvitationImpl::Result{});
 	});
@@ -1779,19 +1779,19 @@ TFuture<TDefaultErrorResult<FLobbiesKickLobbyMemberImpl>> FLobbiesEOSGS::KickLob
 	// Check prerequisites.
 	if (!Services.GetAuthInterface()->IsLoggedIn(Params.LocalAccountId))
 	{
-		UE_LOG(LogTemp, Warning, TEXT("[FLobbiesEOSGS::KickLobbyMemberImpl] Failed: User is not logged in. User[%s]"), *ToLogString(Params.LocalAccountId));
+		UE_LOG(LogOnlineServices, Warning, TEXT("[FLobbiesEOSGS::KickLobbyMemberImpl] Failed: User is not logged in. User[%s]"), *ToLogString(Params.LocalAccountId));
 		return MakeFulfilledPromise<TDefaultErrorResult<FLobbiesKickLobbyMemberImpl>>(Errors::InvalidUser()).GetFuture();
 	}
 
 	if (!Params.LobbyData.IsValid())
 	{
-		UE_LOG(LogTemp, Warning, TEXT("[FLobbiesEOSGS::KickLobbyMemberImpl] Failed: No lobby data provided."));
+		UE_LOG(LogOnlineServices, Warning, TEXT("[FLobbiesEOSGS::KickLobbyMemberImpl] Failed: No lobby data provided."));
 		return MakeFulfilledPromise<TDefaultErrorResult<FLobbiesKickLobbyMemberImpl>>(Errors::InvalidParams()).GetFuture();
 	}
 
 	if (Params.LobbyData->GetLobbyClientData()->GetPublicData().OwnerAccountId != Params.LocalAccountId)
 	{
-		UE_LOG(LogTemp, Warning, TEXT("[FLobbiesEOSGS::KickLobbyMemberImpl] Failed: User is not the lobby owner. User[%s], Lobby[%s]"),
+		UE_LOG(LogOnlineServices, Warning, TEXT("[FLobbiesEOSGS::KickLobbyMemberImpl] Failed: User is not the lobby owner. User[%s], Lobby[%s]"),
 			*ToLogString(Params.LocalAccountId), *Params.LobbyData->GetLobbyIdString());
 		return MakeFulfilledPromise<TDefaultErrorResult<FLobbiesKickLobbyMemberImpl>>(Errors::InvalidParams()).GetFuture();
 	}
@@ -1817,13 +1817,13 @@ TFuture<TDefaultErrorResult<FLobbiesKickLobbyMemberImpl>> FLobbiesEOSGS::KickLob
 	{
 		if (Result->ResultCode != EOS_EResult::EOS_Success)
 		{
-			UE_LOG(LogTemp, Warning, TEXT("[FLobbiesEOSGS::KickLobbyMemberImpl] EOS_Lobby_KickMember Failed: User[%s], TargetUser[%s], Lobby[%s], Result[%s]"),
+			UE_LOG(LogOnlineServices, Warning, TEXT("[FLobbiesEOSGS::KickLobbyMemberImpl] EOS_Lobby_KickMember Failed: User[%s], TargetUser[%s], Lobby[%s], Result[%s]"),
 				*ToLogString(LocalAccountId), *ToLogString(TargetAccountId), *LobbyData->GetLobbyIdString(), *LexToString(Result->ResultCode));
 			Promise.EmplaceValue(Errors::FromEOSResult(Result->ResultCode));
 			return;
 		}
 
-		UE_LOG(LogTemp, Verbose, TEXT("[FLobbiesEOSGS::KickLobbyMemberImpl] Succeeded: User[%s], TargetUser[%s], Lobby[%s]"),
+		UE_LOG(LogOnlineServices, Verbose, TEXT("[FLobbiesEOSGS::KickLobbyMemberImpl] Succeeded: User[%s], TargetUser[%s], Lobby[%s]"),
 			*ToLogString(LocalAccountId), *ToLogString(TargetAccountId), *LobbyData->GetLobbyIdString());
 		Promise.EmplaceValue(FLobbiesKickLobbyMemberImpl::Result{});
 	});
@@ -1836,19 +1836,19 @@ TFuture<TDefaultErrorResult<FLobbiesPromoteLobbyMemberImpl>> FLobbiesEOSGS::Prom
 	// Check prerequisites.
 	if (!Services.GetAuthInterface()->IsLoggedIn(Params.LocalAccountId))
 	{
-		UE_LOG(LogTemp, Warning, TEXT("[FLobbiesEOSGS::PromoteLobbyMemberImpl] Failed: User is not logged in. User[%s]"), *ToLogString(Params.LocalAccountId));
+		UE_LOG(LogOnlineServices, Warning, TEXT("[FLobbiesEOSGS::PromoteLobbyMemberImpl] Failed: User is not logged in. User[%s]"), *ToLogString(Params.LocalAccountId));
 		return MakeFulfilledPromise<TDefaultErrorResult<FLobbiesPromoteLobbyMemberImpl>>(Errors::InvalidUser()).GetFuture();
 	}
 
 	if (!Params.LobbyData.IsValid())
 	{
-		UE_LOG(LogTemp, Warning, TEXT("[FLobbiesEOSGS::PromoteLobbyMemberImpl] Failed: No lobby data provided."));
+		UE_LOG(LogOnlineServices, Warning, TEXT("[FLobbiesEOSGS::PromoteLobbyMemberImpl] Failed: No lobby data provided."));
 		return MakeFulfilledPromise<TDefaultErrorResult<FLobbiesPromoteLobbyMemberImpl>>(Errors::InvalidParams()).GetFuture();
 	}
 
 	if (Params.LobbyData->GetLobbyClientData()->GetPublicData().OwnerAccountId != Params.LocalAccountId)
 	{
-		UE_LOG(LogTemp, Warning, TEXT("[FLobbiesEOSGS::PromoteLobbyMemberImpl] Failed: User is not the lobby owner. User[%s], Lobby[%s]"),
+		UE_LOG(LogOnlineServices, Warning, TEXT("[FLobbiesEOSGS::PromoteLobbyMemberImpl] Failed: User is not the lobby owner. User[%s], Lobby[%s]"),
 			*ToLogString(Params.LocalAccountId), *Params.LobbyData->GetLobbyIdString());
 		return MakeFulfilledPromise<TDefaultErrorResult<FLobbiesPromoteLobbyMemberImpl>>(Errors::InvalidParams()).GetFuture();
 	}
@@ -1874,13 +1874,13 @@ TFuture<TDefaultErrorResult<FLobbiesPromoteLobbyMemberImpl>> FLobbiesEOSGS::Prom
 	{
 		if (Result->ResultCode != EOS_EResult::EOS_Success)
 		{
-			UE_LOG(LogTemp, Warning, TEXT("[FLobbiesEOSGS::PromoteLobbyMemberImpl] EOS_Lobby_PromoteMember Failed: User[%s], TargetUser[%s], Lobby[%s], Result[%s]"),
+			UE_LOG(LogOnlineServices, Warning, TEXT("[FLobbiesEOSGS::PromoteLobbyMemberImpl] EOS_Lobby_PromoteMember Failed: User[%s], TargetUser[%s], Lobby[%s], Result[%s]"),
 				*ToLogString(LocalAccountId), *ToLogString(TargetAccountId), *LobbyData->GetLobbyIdString(), *LexToString(Result->ResultCode));
 			Promise.EmplaceValue(Errors::FromEOSResult(Result->ResultCode));
 			return;
 		}
 
-		UE_LOG(LogTemp, Verbose, TEXT("[FLobbiesEOSGS::PromoteLobbyMemberImpl] Succeeded: User[%s], TargetUser[%s], Lobby[%s]"),
+		UE_LOG(LogOnlineServices, Verbose, TEXT("[FLobbiesEOSGS::PromoteLobbyMemberImpl] Succeeded: User[%s], TargetUser[%s], Lobby[%s]"),
 			*ToLogString(LocalAccountId), *ToLogString(TargetAccountId), *LobbyData->GetLobbyIdString());
 		Promise.EmplaceValue(FLobbiesPromoteLobbyMemberImpl::Result{});
 	});
@@ -1893,20 +1893,20 @@ TFuture<TDefaultErrorResult<FLobbiesModifyLobbyDataImpl>> FLobbiesEOSGS::ModifyL
 	// Check prerequisites.
 	if (!Services.GetAuthInterface()->IsLoggedIn(Params.LocalAccountId))
 	{
-		UE_LOG(LogTemp, Warning, TEXT("[FLobbiesEOSGS::ModifyLobbyDataImpl] Failed: User is not logged in. User[%s]"), *ToLogString(Params.LocalAccountId));
+		UE_LOG(LogOnlineServices, Warning, TEXT("[FLobbiesEOSGS::ModifyLobbyDataImpl] Failed: User is not logged in. User[%s]"), *ToLogString(Params.LocalAccountId));
 		return MakeFulfilledPromise<TDefaultErrorResult<FLobbiesModifyLobbyDataImpl>>(Errors::InvalidUser()).GetFuture();
 	}
 
 	if (!Params.LobbyData.IsValid())
 	{
-		UE_LOG(LogTemp, Warning, TEXT("[FLobbiesEOSGS::ModifyLobbyDataImpl] Failed: No lobby data provided."));
+		UE_LOG(LogOnlineServices, Warning, TEXT("[FLobbiesEOSGS::ModifyLobbyDataImpl] Failed: No lobby data provided."));
 		return MakeFulfilledPromise<TDefaultErrorResult<FLobbiesModifyLobbyDataImpl>>(Errors::InvalidParams()).GetFuture();
 	}
 
 	TSharedPtr<FLobbyDetailsEOS> LobbyDetails = Params.LobbyData->GetUserLobbyDetails(Params.LocalAccountId);
 	if (!LobbyDetails)
 	{
-		UE_LOG(LogTemp, Warning, TEXT("[FLobbiesEOSGS::ModifyLobbyDataImpl] Failed: Unable to find lobby details for user. User[%s], Lobby[%s]"),
+		UE_LOG(LogOnlineServices, Warning, TEXT("[FLobbiesEOSGS::ModifyLobbyDataImpl] Failed: Unable to find lobby details for user. User[%s], Lobby[%s]"),
 			*ToLogString(Params.LocalAccountId), *Params.LobbyData->GetLobbyIdString());
 		return MakeFulfilledPromise<TDefaultErrorResult<FLobbiesModifyLobbyDataImpl>>(Errors::InvalidParams()).GetFuture();
 	}
@@ -1924,13 +1924,13 @@ TFuture<TDefaultErrorResult<FLobbiesModifyLobbyDataImpl>> FLobbiesEOSGS::ModifyL
 	{
 		if (EosResult != EOS_EResult::EOS_Success && EosResult != EOS_EResult::EOS_NoChange)
 		{
-			UE_LOG(LogTemp, Warning, TEXT("[FLobbiesEOSGS::ModifyLobbyDataImpl] ApplyLobbyDataUpdateFromLocalChanges Failed: User[%s], Lobby[%s], Result[%s]"),
+			UE_LOG(LogOnlineServices, Warning, TEXT("[FLobbiesEOSGS::ModifyLobbyDataImpl] ApplyLobbyDataUpdateFromLocalChanges Failed: User[%s], Lobby[%s], Result[%s]"),
 				*ToLogString(LocalAccountId), *LobbyData->GetLobbyIdString(), *LexToString(EosResult));
 			Promise.EmplaceValue(Errors::FromEOSResult(EosResult));
 			return;
 		}
 
-		UE_LOG(LogTemp, Verbose, TEXT("[FLobbiesEOSGS::ModifyLobbyDataImpl] Succeeded: User[%s], Lobby[%s]"),
+		UE_LOG(LogOnlineServices, Verbose, TEXT("[FLobbiesEOSGS::ModifyLobbyDataImpl] Succeeded: User[%s], Lobby[%s]"),
 			*ToLogString(LocalAccountId), *LobbyData->GetLobbyIdString());
 		Promise.EmplaceValue(FLobbiesModifyLobbyDataImpl::Result{});
 	});
@@ -1946,7 +1946,7 @@ TOnlineAsyncOpHandle<FLobbiesProcessLobbyNotificationImpl> FLobbiesEOSGS::Proces
 	// Check prerequisites.
 	if (!Params.LobbyData.IsValid())
 	{
-		UE_LOG(LogTemp, Warning, TEXT("[FLobbiesEOSGS::ProcessLobbyNotificationImplOp] Failed: No lobby data provided."));
+		UE_LOG(LogOnlineServices, Warning, TEXT("[FLobbiesEOSGS::ProcessLobbyNotificationImplOp] Failed: No lobby data provided."));
 		Op->SetError(Errors::InvalidParams());
 		return Op->GetHandle();
 	}
@@ -1961,7 +1961,7 @@ TOnlineAsyncOpHandle<FLobbiesProcessLobbyNotificationImpl> FLobbiesEOSGS::Proces
 		TSharedPtr<FLobbyDetailsEOS> LobbyDetails = Params.LobbyData->GetActiveLobbyDetails();
 		if (!LobbyDetails.IsValid())
 		{
-			UE_LOG(LogTemp, Warning, TEXT("[FLobbiesEOSGS::ProcessLobbyNotificationImplOp] Failed: Unable to find active lobby details to process notificaions. Lobby[%s]"),
+			UE_LOG(LogOnlineServices, Warning, TEXT("[FLobbiesEOSGS::ProcessLobbyNotificationImplOp] Failed: Unable to find active lobby details to process notificaions. Lobby[%s]"),
 				*Params.LobbyData->GetLobbyIdString());
 			InAsyncOp.SetError(Errors::InvalidState());
 			return MakeFulfilledPromise<TDefaultErrorResultInternal<FLobbyServiceSnapshot>>(Errors::InvalidState()).GetFuture();
@@ -1979,7 +1979,7 @@ TOnlineAsyncOpHandle<FLobbiesProcessLobbyNotificationImpl> FLobbiesEOSGS::Proces
 
 		if (LobbySnapshotResult.IsError())
 		{
-			UE_LOG(LogTemp, Warning, TEXT("[FLobbiesEOSGS::ProcessLobbyNotificationImplOp] GetLobbySnapshot Failed. Lobby[%s], Result[%s]"),
+			UE_LOG(LogOnlineServices, Warning, TEXT("[FLobbiesEOSGS::ProcessLobbyNotificationImplOp] GetLobbySnapshot Failed. Lobby[%s], Result[%s]"),
 				*Params.LobbyData->GetLobbyIdString(), *LobbySnapshotResult.GetErrorValue().GetLogString());
 			InAsyncOp.SetError(MoveTemp(LobbySnapshotResult.GetErrorValue()));
 			return;
@@ -1996,7 +1996,7 @@ TOnlineAsyncOpHandle<FLobbiesProcessLobbyNotificationImpl> FLobbiesEOSGS::Proces
 				TDefaultErrorResultInternal<FLobbyMemberServiceSnapshot> MemberSnapshotResult = LobbyDetails->GetLobbyMemberSnapshot(MutatedMemberAccountId);
 				if (MemberSnapshotResult.IsError())
 				{
-					UE_LOG(LogTemp, Warning, TEXT("[FLobbiesEOSGS::ProcessLobbyNotificationImplOp] GetLobbyMemberSnapshot Failed. User[%s], Lobby[%s], Result[%s]"),
+					UE_LOG(LogOnlineServices, Warning, TEXT("[FLobbiesEOSGS::ProcessLobbyNotificationImplOp] GetLobbyMemberSnapshot Failed. User[%s], Lobby[%s], Result[%s]"),
 						*ToLogString(MutatedMemberAccountId), *Params.LobbyData->GetLobbyIdString(), *LobbySnapshotResult.GetErrorValue().GetLogString());
 					InAsyncOp.SetError(MoveTemp(MemberSnapshotResult.GetErrorValue()));
 					return;
@@ -2022,7 +2022,7 @@ TOnlineAsyncOpHandle<FLobbiesProcessLobbyNotificationImpl> FLobbiesEOSGS::Proces
 			Params.LobbyData->GetLobbyClientData()->PrepareServiceSnapshot({ MoveTemp(LobbySnapshotResult.GetOkValue()), MoveTemp(LobbyMemberSnapshots), MoveTemp(LeavingMemberReason) });
 		if (PrepareSnapshotResult.IsError())
 		{
-			UE_LOG(LogTemp, Warning, TEXT("[FLobbiesEOSGS::ProcessLobbyNotificationImplOp] PrepareServiceSnapshot Failed. Lobby[%s], Result[%s]"),
+			UE_LOG(LogOnlineServices, Warning, TEXT("[FLobbiesEOSGS::ProcessLobbyNotificationImplOp] PrepareServiceSnapshot Failed. Lobby[%s], Result[%s]"),
 				*Params.LobbyData->GetLobbyIdString(), *PrepareSnapshotResult.GetErrorValue().GetLogString());
 			InAsyncOp.SetError(MoveTemp(PrepareSnapshotResult.GetErrorValue()));
 			return;
@@ -2041,7 +2041,7 @@ TOnlineAsyncOpHandle<FLobbiesProcessLobbyNotificationImpl> FLobbiesEOSGS::Proces
 			RemoveActiveLobby(LeavingMember, Params.LobbyData.ToSharedRef());
 		}
 
-		UE_LOG(LogTemp, Verbose, TEXT("[FLobbiesEOSGS::ProcessLobbyNotificationImplOp] Succeeded: Lobby[%s]"),
+		UE_LOG(LogOnlineServices, Verbose, TEXT("[FLobbiesEOSGS::ProcessLobbyNotificationImplOp] Succeeded: Lobby[%s]"),
 			*Params.LobbyData->GetLobbyIdString());
 		InAsyncOp.SetResult(FLobbiesProcessLobbyNotificationImpl::Result{});
 	})
diff --git a/Engine/Plugins/Online/OnlineServicesEOSGS/Source/Private/Online/LobbiesEOSGSTypes.cpp b/Engine/Plugins/Online/OnlineServicesEOSGS/Source/Private/Online/LobbiesEOSGSTypes.cpp
--- a/Engine/Plugins/Online/OnlineServicesEOSGS/Source/Private/Online/LobbiesEOSGSTypes.cpp
+++ b/Engine/Plugins/Online/OnlineServicesEOSGS/Source/Private/Online/LobbiesEOSGSTypes.cpp
@@ -164,7 +164,7 @@ TDefaultErrorResultInternal<TSharedRef<FLobbyDetailsEOS>> FLobbyDetailsEOS::Crea
 	EOS_EResult EOSResult = EOS_Lobby_CopyLobbyDetailsHandle(Prerequisites->LobbyInterfaceHandle, &Options, &LobbyDetailsHandle);
 	if (EOSResult != EOS_EResult::EOS_Success)
 	{
-		UE_LOG(LogTemp, Warning, TEXT("[FLobbyDetailsEOS::CreateFromLobbyId] EOS_Lobby_CopyLobbyDetailsHandle Failed: User[%s], Lobby[%s], Result[%s]"),
+		UE_LOG(LogOnlineServices, Warning, TEXT("[FLobbyDetailsEOS::CreateFromLobbyId] EOS_Lobby_CopyLobbyDetailsHandle Failed: User[%s], Lobby[%s], Result[%s]"),
 			*ToLogString(LocalAccountId), UTF8_TO_TCHAR(LobbyIdEOS), *ToLogString(EOSResult));
 		return TDefaultErrorResultInternal<TSharedRef<FLobbyDetailsEOS>>(Errors::FromEOSResult(EOSResult));
 	}
@@ -172,13 +172,13 @@ TDefaultErrorResultInternal<TSharedRef<FLobbyDetailsEOS>> FLobbyDetailsEOS::Crea
 	TDefaultErrorResultInternal<TSharedRef<FLobbyDetailsInfoEOS>> LobbyDetailsInfoResult = FLobbyDetailsInfoEOS::Create(LobbyDetailsHandle);
 	if (LobbyDetailsInfoResult.IsError())
 	{
-		UE_LOG(LogTemp, Warning, TEXT("[FLobbyDetailsEOS::CreateFromLobbyId] FLobbyDetailsInfoEOS::Create Failed: User[%s], Lobby[%s], Result[%s]"),
+		UE_LOG(LogOnlineServices, Warning, TEXT("[FLobbyDetailsEOS::CreateFromLobbyId] FLobbyDetailsInfoEOS::Create Failed: User[%s], Lobby[%s], Result[%s]"),
 			*ToLogString(LocalAccountId), UTF8_TO_TCHAR(LobbyIdEOS), *ToLogString(EOSResult));
 		EOS_LobbyDetails_Release(LobbyDetailsHandle);
 		return TDefaultErrorResultInternal<TSharedRef<FLobbyDetailsEOS>>(MoveTemp(LobbyDetailsInfoResult.GetErrorValue()));
 	}
 
-	UE_LOG(LogTemp, VeryVerbose, TEXT("[FLobbyDetailsEOS::CreateFromLobbyId] Succeeded: User[%s], Lobby[%s]"),
+	UE_LOG(LogOnlineServices, VeryVerbose, TEXT("[FLobbyDetailsEOS::CreateFromLobbyId] Succeeded: User[%s], Lobby[%s]"),
 		*ToLogString(LocalAccountId), UTF8_TO_TCHAR(LobbyIdEOS));
 	return TDefaultErrorResultInternal<TSharedRef<FLobbyDetailsEOS>>(MakeShared<FLobbyDetailsEOS>(Prerequisites, LobbyDetailsInfoResult.GetOkValue(), LocalAccountId, ELobbyDetailsSource::Active, LobbyDetailsHandle));
 }
@@ -206,7 +206,7 @@ TDefaultErrorResultInternal<TSharedRef<FLobbyDetailsEOS>> FLobbyDetailsEOS::Crea
 	EOS_EResult EOSResult = EOS_Lobby_CopyLobbyDetailsHandleByInviteId(Prerequisites->LobbyInterfaceHandle, &Options, &LobbyDetailsHandle);
 	if (EOSResult != EOS_EResult::EOS_Success)
 	{
-		UE_LOG(LogTemp, Warning, TEXT("[FLobbyDetailsEOS::CreateFromInviteId] EOS_Lobby_CopyLobbyDetailsHandleByInviteId Failed: User[%s], InviteId[%s], Result[%s]"),
+		UE_LOG(LogOnlineServices, Warning, TEXT("[FLobbyDetailsEOS::CreateFromInviteId] EOS_Lobby_CopyLobbyDetailsHandleByInviteId Failed: User[%s], InviteId[%s], Result[%s]"),
 			*ToLogString(LocalAccountId), UTF8_TO_TCHAR(InviteId), *ToLogString(EOSResult));
 		return TDefaultErrorResultInternal<TSharedRef<FLobbyDetailsEOS>>(Errors::FromEOSResult(EOSResult));
 	}
@@ -214,13 +214,13 @@ TDefaultErrorResultInternal<TSharedRef<FLobbyDetailsEOS>> FLobbyDetailsEOS::Crea
 	TDefaultErrorResultInternal<TSharedRef<FLobbyDetailsInfoEOS>> LobbyDetailsInfoResult = FLobbyDetailsInfoEOS::Create(LobbyDetailsHandle);
 	if (LobbyDetailsInfoResult.IsError())
 	{
-		UE_LOG(LogTemp, Warning, TEXT("[FLobbyDetailsEOS::CreateFromInviteId] FLobbyDetailsInfoEOS::Create Failed: User[%s], InviteId[%s], Result[%s]"),
+		UE_LOG(LogOnlineServices, Warning, TEXT("[FLobbyDetailsEOS::CreateFromInviteId] FLobbyDetailsInfoEOS::Create Failed: User[%s], InviteId[%s], Result[%s]"),
 			*ToLogString(LocalAccountId), UTF8_TO_TCHAR(InviteId), *LobbyDetailsInfoResult.GetErrorValue().GetLogString());
 		EOS_LobbyDetails_Release(LobbyDetailsHandle);
 		return TDefaultErrorResultInternal<TSharedRef<FLobbyDetailsEOS>>(MoveTemp(LobbyDetailsInfoResult.GetErrorValue()));
 	}
 
-	UE_LOG(LogTemp, VeryVerbose, TEXT("[FLobbyDetailsEOS::CreateFromInviteId] Succeeded: User[%s], InviteId[%s], Lobby[%s]"),
+	UE_LOG(LogOnlineServices, VeryVerbose, TEXT("[FLobbyDetailsEOS::CreateFromInviteId] Succeeded: User[%s], InviteId[%s], Lobby[%s]"),
 		*ToLogString(LocalAccountId), UTF8_TO_TCHAR(InviteId), UTF8_TO_TCHAR(LobbyDetailsInfoResult.GetOkValue()->GetLobbyIdEOS()));
 	return TDefaultErrorResultInternal<TSharedRef<FLobbyDetailsEOS>>(MakeShared<FLobbyDetailsEOS>(Prerequisites, LobbyDetailsInfoResult.GetOkValue(), LocalAccountId, ELobbyDetailsSource::Invite, LobbyDetailsHandle));
 }
@@ -240,7 +240,7 @@ TDefaultErrorResultInternal<TSharedRef<FLobbyDetailsEOS>> FLobbyDetailsEOS::Crea
 	EOS_EResult EOSResult = EOS_Lobby_CopyLobbyDetailsHandleByUiEventId(Prerequisites->LobbyInterfaceHandle, &Options, &LobbyDetailsHandle);
 	if (EOSResult != EOS_EResult::EOS_Success)
 	{
-		UE_LOG(LogTemp, Warning, TEXT("[FLobbyDetailsEOS::CreateFromUiEventId] EOS_Lobby_CopyLobbyDetailsHandleByUiEventId Failed: User[%s], Result[%s]"),
+		UE_LOG(LogOnlineServices, Warning, TEXT("[FLobbyDetailsEOS::CreateFromUiEventId] EOS_Lobby_CopyLobbyDetailsHandleByUiEventId Failed: User[%s], Result[%s]"),
 			*ToLogString(LocalAccountId), *ToLogString(EOSResult));
 		return TDefaultErrorResultInternal<TSharedRef<FLobbyDetailsEOS>>(Errors::FromEOSResult(EOSResult));
 	}
@@ -248,13 +248,13 @@ TDefaultErrorResultInternal<TSharedRef<FLobbyDetailsEOS>> FLobbyDetailsEOS::Crea
 	TDefaultErrorResultInternal<TSharedRef<FLobbyDetailsInfoEOS>> LobbyDetailsInfoResult = FLobbyDetailsInfoEOS::Create(LobbyDetailsHandle);
 	if (LobbyDetailsInfoResult.IsError())
 	{
-		UE_LOG(LogTemp, Warning, TEXT("[FLobbyDetailsEOS::CreateFromUiEventId] FLobbyDetailsInfoEOS::Create Failed: User[%s], Result[%s]"),
+		UE_LOG(LogOnlineServices, Warning, TEXT("[FLobbyDetailsEOS::CreateFromUiEventId] FLobbyDetailsInfoEOS::Create Failed: User[%s], Result[%s]"),
 			*ToLogString(LocalAccountId), *LobbyDetailsInfoResult.GetErrorValue().GetLogString());
 		EOS_LobbyDetails_Release(LobbyDetailsHandle);
 		return TDefaultErrorResultInternal<TSharedRef<FLobbyDetailsEOS>>(MoveTemp(LobbyDetailsInfoResult.GetErrorValue()));
 	}
 
-	UE_LOG(LogTemp, VeryVerbose, TEXT("[FLobbyDetailsEOS::CreateFromUiEventId] Succeeded: User[%s], Lobby[%s]"),
+	UE_LOG(LogOnlineServices, VeryVerbose, TEXT("[FLobbyDetailsEOS::CreateFromUiEventId] Succeeded: User[%s], Lobby[%s]"),
 		*ToLogString(LocalAccountId), UTF8_TO_TCHAR(LobbyDetailsInfoResult.GetOkValue()->GetLobbyIdEOS()));
 	return TDefaultErrorResultInternal<TSharedRef<FLobbyDetailsEOS>>(MakeShared<FLobbyDetailsEOS>(Prerequisites, LobbyDetailsInfoResult.GetOkValue(), LocalAccountId, ELobbyDetailsSource::UiEvent, LobbyDetailsHandle));
 }
@@ -275,7 +275,7 @@ TDefaultErrorResultInternal<TSharedRef<FLobbyDetailsEOS>> FLobbyDetailsEOS::Crea
 	EOS_EResult EOSResult = EOS_LobbySearch_CopySearchResultByIndex(SearchHandle, &Options, &LobbyDetailsHandle);
 	if (EOSResult != EOS_EResult::EOS_Success)
 	{
-		UE_LOG(LogTemp, Warning, TEXT("[FLobbyDetailsEOS::CreateFromSearchResult] EOS_LobbySearch_CopySearchResultByIndex Failed: User[%s], Result[%s]"),
+		UE_LOG(LogOnlineServices, Warning, TEXT("[FLobbyDetailsEOS::CreateFromSearchResult] EOS_LobbySearch_CopySearchResultByIndex Failed: User[%s], Result[%s]"),
 			*ToLogString(LocalAccountId), *ToLogString(EOSResult));
 		return TDefaultErrorResultInternal<TSharedRef<FLobbyDetailsEOS>>(Errors::FromEOSResult(EOSResult));
 	}
@@ -283,13 +283,13 @@ TDefaultErrorResultInternal<TSharedRef<FLobbyDetailsEOS>> FLobbyDetailsEOS::Crea
 	TDefaultErrorResultInternal<TSharedRef<FLobbyDetailsInfoEOS>> LobbyDetailsInfoResult = FLobbyDetailsInfoEOS::Create(LobbyDetailsHandle);
 	if (LobbyDetailsInfoResult.IsError())
 	{
-		UE_LOG(LogTemp, Warning, TEXT("[FLobbyDetailsEOS::CreateFromSearchResult] FLobbyDetailsInfoEOS::Create Failed: User[%s], Result[%s]"),
+		UE_LOG(LogOnlineServices, Warning, TEXT("[FLobbyDetailsEOS::CreateFromSearchResult] FLobbyDetailsInfoEOS::Create Failed: User[%s], Result[%s]"),
 			*ToLogString(LocalAccountId), *LobbyDetailsInfoResult.GetErrorValue().GetLogString());
 		EOS_LobbyDetails_Release(LobbyDetailsHandle);
 		return TDefaultErrorResultInternal<TSharedRef<FLobbyDetailsEOS>>(MoveTemp(LobbyDetailsInfoResult.GetErrorValue()));
 	}
 
-	UE_LOG(LogTemp, VeryVerbose, TEXT("[FLobbyDetailsEOS::CreateFromSearchResult] Succeeded: User[%s], Lobby[%s]"),
+	UE_LOG(LogOnlineServices, VeryVerbose, TEXT("[FLobbyDetailsEOS::CreateFromSearchResult] Succeeded: User[%s], Lobby[%s]"),
 		*ToLogString(LocalAccountId), UTF8_TO_TCHAR(LobbyDetailsInfoResult.GetOkValue()->GetLobbyIdEOS()));
 	return TDefaultErrorResultInternal<TSharedRef<FLobbyDetailsEOS>>(MakeShared<FLobbyDetailsEOS>(Prerequisites, LobbyDetailsInfoResult.GetOkValue(), LocalAccountId, ELobbyDetailsSource::Search, LobbyDetailsHandle));
 }
@@ -309,7 +309,7 @@ TFuture<TDefaultErrorResultInternal<FLobbyServiceSnapshot>> FLobbyDetailsEOS::Ge
 	TSharedPtr<FAuthEOSGS> AuthInterface = Prerequisites->AuthInterface.Pin();
 	if (!AuthInterface)
 	{
-		UE_LOG(LogTemp, Warning, TEXT("[FLobbyDetailsEOS::GetLobbySnapshot] Failed: Auth interface missing. Lobby[%s]"),
+		UE_LOG(LogOnlineServices, Warning, TEXT("[FLobbyDetailsEOS::GetLobbySnapshot] Failed: Auth interface missing. Lobby[%s]"),
 			UTF8_TO_TCHAR(GetInfo()->GetLobbyIdEOS()));
 		return MakeFulfilledPromise<TDefaultErrorResultInternal<FLobbyServiceSnapshot>>(Errors::MissingInterface()).GetFuture();
 	}
@@ -342,7 +342,7 @@ TFuture<TDefaultErrorResultInternal<FLobbyServiceSnapshot>> FLobbyDetailsEOS::Ge
 	{
 		if (MemberProductUserIds->Num() != ResolvedAccountIds.Num())
 		{
-			UE_LOG(LogTemp, Warning, TEXT("[FLobbyDetailsEOS::GetLobbySnapshot] ResolveAccountIds Failed: Expected Count[%d], Received Count[%d], Lobby[%s]"),
+			UE_LOG(LogOnlineServices, Warning, TEXT("[FLobbyDetailsEOS::GetLobbySnapshot] ResolveAccountIds Failed: Expected Count[%d], Received Count[%d], Lobby[%s]"),
 				MemberProductUserIds->Num(), ResolvedAccountIds.Num(), UTF8_TO_TCHAR(StrongThis->GetInfo()->GetLobbyIdEOS()));
 			Promise.EmplaceValue(Errors::Unknown());
 			return;
@@ -397,7 +397,7 @@ TFuture<TDefaultErrorResultInternal<FLobbyServiceSnapshot>> FLobbyDetailsEOS::Ge
 				EOS_EResult EOSResult = EOS_LobbyDetails_CopyAttributeByIndex(StrongThis->LobbyDetailsHandle, &CopyAttributeByIndexOptions, &LobbyAttribute);
 				if (EOSResult != EOS_EResult::EOS_Success)
 				{
-					UE_LOG(LogTemp, Warning, TEXT("[FLobbyDetailsEOS::GetLobbySnapshot] EOS_LobbyDetails_CopyAttributeByIndex Failed: Lobby[%s], Result[%s]"),
+					UE_LOG(LogOnlineServices, Warning, TEXT("[FLobbyDetailsEOS::GetLobbySnapshot] EOS_LobbyDetails_CopyAttributeByIndex Failed: Lobby[%s], Result[%s]"),
 						UTF8_TO_TCHAR(StrongThis->GetInfo()->GetLobbyIdEOS()), *LexToString(EOSResult));
 					Promise.EmplaceValue(Errors::FromEOSResult(EOSResult));
 					return;
@@ -408,7 +408,7 @@ TFuture<TDefaultErrorResultInternal<FLobbyServiceSnapshot>> FLobbyDetailsEOS::Ge
 			}
 		}
 
-		UE_LOG(LogTemp, VeryVerbose, TEXT("[FLobbyDetailsEOS::GetLobbySnapshot] Succeeded: Lobby[%s]"),
+		UE_LOG(LogOnlineServices, VeryVerbose, TEXT("[FLobbyDetailsEOS::GetLobbySnapshot] Succeeded: Lobby[%s]"),
 			UTF8_TO_TCHAR(StrongThis->GetInfo()->GetLobbyIdEOS()));
 		Promise.EmplaceValue(MoveTemp(LobbyServiceSnapshot));
 	});
@@ -448,7 +448,7 @@ TDefaultErrorResultInternal<FLobbyMemberServiceSnapshot> FLobbyDetailsEOS::GetLo
 			EOS_EResult EOSResult = EOS_LobbyDetails_CopyMemberAttributeByIndex(LobbyDetailsHandle, &CopyMemberAttributeByIndexOptions, &LobbyAttribute);
 			if (EOSResult != EOS_EResult::EOS_Success)
 			{
-				UE_LOG(LogTemp, Warning, TEXT("[FLobbyDetailsEOS::GetLobbyMemberSnapshot] EOS_LobbyDetails_CopyMemberAttributeByIndex Failed: Member[%s], Lobby[%s], Result[%s]"),
+				UE_LOG(LogOnlineServices, Warning, TEXT("[FLobbyDetailsEOS::GetLobbyMemberSnapshot] EOS_LobbyDetails_CopyMemberAttributeByIndex Failed: Member[%s], Lobby[%s], Result[%s]"),
 					*ToLogString(MemberAccountId), UTF8_TO_TCHAR(GetInfo()->GetLobbyIdEOS()), *LexToString(EOSResult));
 				return TDefaultErrorResultInternal<FLobbyMemberServiceSnapshot>(Errors::FromEOSResult(EOSResult));
 			}
@@ -458,7 +458,7 @@ TDefaultErrorResultInternal<FLobbyMemberServiceSnapshot> FLobbyDetailsEOS::GetLo
 		}
 	}
 
-	UE_LOG(LogTemp, VeryVerbose, TEXT("[FLobbyDetailsEOS::GetLobbyMemberSnapshot] Succeeded: Member[%s], Lobby[%s]"),
+	UE_LOG(LogOnlineServices, VeryVerbose, TEXT("[FLobbyDetailsEOS::GetLobbyMemberSnapshot] Succeeded: Member[%s], Lobby[%s]"),
 		*ToLogString(MemberAccountId), UTF8_TO_TCHAR(GetInfo()->GetLobbyIdEOS()));
 	return TDefaultErrorResultInternal<FLobbyMemberServiceSnapshot>(MoveTemp(LobbyMemberServcieSnapshot));
 }
@@ -482,7 +482,7 @@ TFuture<EOS_EResult> FLobbyDetailsEOS::ApplyLobbyDataUpdateFromLocalChanges(FAcc
 	EOS_EResult EOSResult = EOS_Lobby_UpdateLobbyModification(Prerequisites->LobbyInterfaceHandle, &ModificationOptions, &LobbyModificationHandle);
 	if (EOSResult != EOS_EResult::EOS_Success)
 	{
-		UE_LOG(LogTemp, Warning, TEXT("[FLobbyDetailsEOS::ApplyLobbyDataUpdateFromLocalChanges] EOS_Lobby_UpdateLobbyModification Failed: User[%s], Lobby[%s], Result[%s]"),
+		UE_LOG(LogOnlineServices, Warning, TEXT("[FLobbyDetailsEOS::ApplyLobbyDataUpdateFromLocalChanges] EOS_Lobby_UpdateLobbyModification Failed: User[%s], Lobby[%s], Result[%s]"),
 			*ToLogString(LocalAccountId), UTF8_TO_TCHAR(GetInfo()->GetLobbyIdEOS()), *LexToString(EOSResult));
 		return MakeFulfilledPromise<EOS_EResult>(EOSResult).GetFuture();
 	}
@@ -498,7 +498,7 @@ TFuture<EOS_EResult> FLobbyDetailsEOS::ApplyLobbyDataUpdateFromLocalChanges(FAcc
 		EOSResult = EOS_LobbyModification_SetPermissionLevel(LobbyModificationHandle, &SetPermissionOptions);
 		if (EOSResult != EOS_EResult::EOS_Success)
 		{
-			UE_LOG(LogTemp, Warning, TEXT("[FLobbyDetailsEOS::ApplyLobbyDataUpdateFromLocalChanges] EOS_LobbyModification_SetPermissionLevel Failed: User[%s], Lobby[%s], Result[%s]"),
+			UE_LOG(LogOnlineServices, Warning, TEXT("[FLobbyDetailsEOS::ApplyLobbyDataUpdateFromLocalChanges] EOS_LobbyModification_SetPermissionLevel Failed: User[%s], Lobby[%s], Result[%s]"),
 				*ToLogString(LocalAccountId), UTF8_TO_TCHAR(GetInfo()->GetLobbyIdEOS()), *LexToString(EOSResult));
 			return MakeFulfilledPromise<EOS_EResult>(EOSResult).GetFuture();
 		}
@@ -518,7 +518,7 @@ TFuture<EOS_EResult> FLobbyDetailsEOS::ApplyLobbyDataUpdateFromLocalChanges(FAcc
 		EOSResult = EOS_LobbyModification_AddAttribute(LobbyModificationHandle, &AddAttributeOptions);
 		if (EOSResult != EOS_EResult::EOS_Success)
 		{
-			UE_LOG(LogTemp, Warning, TEXT("[FLobbyDetailsEOS::ApplyLobbyDataUpdateFromLocalChanges] EOS_LobbyModification_AddAttribute Failed: User[%s], Lobby[%s], Key[%s], Value[%s], Result[%s]"),
+			UE_LOG(LogOnlineServices, Warning, TEXT("[FLobbyDetailsEOS::ApplyLobbyDataUpdateFromLocalChanges] EOS_LobbyModification_AddAttribute Failed: User[%s], Lobby[%s], Key[%s], Value[%s], Result[%s]"),
 				*ToLogString(LocalAccountId), UTF8_TO_TCHAR(GetInfo()->GetLobbyIdEOS()), *UpdatedAttribute.Key.ToString().ToLower(), *ToLogString(UpdatedAttribute.Value.Value), *LexToString(EOSResult));
 			return MakeFulfilledPromise<EOS_EResult>(EOSResult).GetFuture();
 		}
@@ -537,7 +537,7 @@ TFuture<EOS_EResult> FLobbyDetailsEOS::ApplyLobbyDataUpdateFromLocalChanges(FAcc
 		EOSResult = EOS_LobbyModification_RemoveAttribute(LobbyModificationHandle, &RemoveAttributeOptions);
 		if (EOSResult != EOS_EResult::EOS_Success)
 		{
-			UE_LOG(LogTemp, Warning, TEXT("[FLobbyDetailsEOS::ApplyLobbyDataUpdateFromLocalChanges] EOS_LobbyModification_RemoveAttribute Failed: User[%s], Lobby[%s], Key[%s], Result[%s]"),
+			UE_LOG(LogOnlineServices, Warning, TEXT("[FLobbyDetailsEOS::ApplyLobbyDataUpdateFromLocalChanges] EOS_LobbyModification_RemoveAttribute Failed: User[%s], Lobby[%s], Key[%s], Result[%s]"),
 				*ToLogString(LocalAccountId), UTF8_TO_TCHAR(GetInfo()->GetLobbyIdEOS()), *RemovedAttribute.ToString().ToLower(), *LexToString(EOSResult));
 			return MakeFulfilledPromise<EOS_EResult>(EOSResult).GetFuture();
 		}
@@ -557,7 +557,7 @@ TFuture<EOS_EResult> FLobbyDetailsEOS::ApplyLobbyDataUpdateFromLocalChanges(FAcc
 		EOSResult = EOS_LobbyModification_AddMemberAttribute(LobbyModificationHandle, &AddMemberAttributeOptions);
 		if (EOSResult != EOS_EResult::EOS_Success)
 		{
-			UE_LOG(LogTemp, Warning, TEXT("[FLobbyDetailsEOS::ApplyLobbyMemberDataUpdateFromLocalChanges] EOS_LobbyModification_AddMemberAttribute Failed: User[%s], Lobby[%s], Key[%s], Value[%s], Result[%s]"),
+			UE_LOG(LogOnlineServices, Warning, TEXT("[FLobbyDetailsEOS::ApplyLobbyMemberDataUpdateFromLocalChanges] EOS_LobbyModification_AddMemberAttribute Failed: User[%s], Lobby[%s], Key[%s], Value[%s], Result[%s]"),
 				*ToLogString(LocalAccountId), UTF8_TO_TCHAR(GetInfo()->GetLobbyIdEOS()), *UpdatedAttribute.Key.ToString().ToLower(), *ToLogString(UpdatedAttribute.Value.Value), *LexToString(EOSResult));
 			return MakeFulfilledPromise<EOS_EResult>(EOSResult).GetFuture();
 		}
@@ -576,7 +576,7 @@ TFuture<EOS_EResult> FLobbyDetailsEOS::ApplyLobbyDataUpdateFromLocalChanges(FAcc
 		EOSResult = EOS_LobbyModification_RemoveMemberAttribute(LobbyModificationHandle, &RemoveMemberAttributeOptions);
 		if (EOSResult != EOS_EResult::EOS_Success)
 		{
-			UE_LOG(LogTemp, Warning, TEXT("[FLobbyDetailsEOS::ApplyLobbyMemberDataUpdateFromLocalChanges] EOS_LobbyModification_RemoveMemberAttribute Failed: User[%s], Lobby[%s], Key[%s], Result[%s]"),
+			UE_LOG(LogOnlineServices, Warning, TEXT("[FLobbyDetailsEOS::ApplyLobbyMemberDataUpdateFromLocalChanges] EOS_LobbyModification_RemoveMemberAttribute Failed: User[%s], Lobby[%s], Key[%s], Result[%s]"),
 				*ToLogString(LocalAccountId), UTF8_TO_TCHAR(GetInfo()->GetLobbyIdEOS()), *RemovedAttribute.ToString().ToLower(), *LexToString(EOSResult));
 			return MakeFulfilledPromise<EOS_EResult>(EOSResult).GetFuture();
 		}
@@ -596,12 +596,12 @@ TFuture<EOS_EResult> FLobbyDetailsEOS::ApplyLobbyDataUpdateFromLocalChanges(FAcc
 	{
 		if (Data->ResultCode != EOS_EResult::EOS_Success && Data->ResultCode != EOS_EResult::EOS_NoChange)
 		{
-			UE_LOG(LogTemp, Warning, TEXT("[FLobbyDetailsEOS::ApplyLobbyDataUpdateFromLocalChanges] EOS_Lobby_UpdateLobby Failed: User[%s], Lobby[%s], Result[%s]"),
+			UE_LOG(LogOnlineServices, Warning, TEXT("[FLobbyDetailsEOS::ApplyLobbyDataUpdateFromLocalChanges] EOS_Lobby_UpdateLobby Failed: User[%s], Lobby[%s], Result[%s]"),
 				*ToLogString(LocalAccountId), UTF8_TO_TCHAR(LobbyInfo->GetLobbyIdEOS()), *LexToString(Data->ResultCode));
 		}
 		else
 		{
-			UE_LOG(LogTemp, VeryVerbose, TEXT("[FLobbyDetailsEOS::ApplyLobbyDataUpdateFromLocalChanges] Succeeded: User[%s], Lobby[%s], Result[%s]"),
+			UE_LOG(LogOnlineServices, VeryVerbose, TEXT("[FLobbyDetailsEOS::ApplyLobbyDataUpdateFromLocalChanges] Succeeded: User[%s], Lobby[%s], Result[%s]"),
 				*ToLogString(LocalAccountId), UTF8_TO_TCHAR(LobbyInfo->GetLobbyIdEOS()), *LexToString(Data->ResultCode));
 		}
 
@@ -634,7 +634,7 @@ TDefaultErrorResultInternal<TSharedRef<FLobbyDetailsInfoEOS>> FLobbyDetailsInfoE
 	EOS_EResult EOSResult = EOS_LobbyDetails_CopyInfo(LobbyDetailsHandle, &CopyInfoOptions, &LobbyDetailsInfo);
 	if (EOSResult != EOS_EResult::EOS_Success)
 	{
-		UE_LOG(LogTemp, Warning, TEXT("[FLobbyDetailsInfoEOS] EOS_LobbyDetails_CopyInfo Failed: Result[%s]"),
+		UE_LOG(LogOnlineServices, Warning, TEXT("[FLobbyDetailsInfoEOS] EOS_LobbyDetails_CopyInfo Failed: Result[%s]"),
 			*LexToString(EOSResult));
 		return TDefaultErrorResultInternal<TSharedRef<FLobbyDetailsInfoEOS>>(Errors::FromEOSResult(EOSResult));
 	}
@@ -650,7 +650,7 @@ FLobbyDetailsInfoEOS::FLobbyDetailsInfoEOS(FLobbyDetailsInfoPtr&& InLobbyDetails
 
 	if (!BucketId.IsValid())
 	{
-		UE_LOG(LogTemp, Warning, TEXT("[FLobbyDetailsInfoEOS] Failed: Unable to parse lobby bucket id. Lobby[%s], Bucket[%s]"),
+		UE_LOG(LogOnlineServices, Warning, TEXT("[FLobbyDetailsInfoEOS] Failed: Unable to parse lobby bucket id. Lobby[%s], Bucket[%s]"),
 			UTF8_TO_TCHAR(LobbyDetailsInfo->LobbyId), UTF8_TO_TCHAR(LobbyDetailsInfo->BucketId));
 	}
 }
@@ -731,7 +731,7 @@ TFuture<TDefaultErrorResultInternal<TSharedRef<FLobbyDataEOS>>> FLobbyDataEOS::C
 	{
 		if (Result.IsError())
 		{
-			UE_LOG(LogTemp, Warning, TEXT("[FLobbyDataEOS::Create] FLobbyDetailsEOS::GetLobbySnapshot Failed: Lobby[%s], Result[%s]"),
+			UE_LOG(LogOnlineServices, Warning, TEXT("[FLobbyDataEOS::Create] FLobbyDetailsEOS::GetLobbySnapshot Failed: Lobby[%s], Result[%s]"),
 				*ToLogString(LobbyId), *Result.GetErrorValue().GetLogString());
 			Promise.EmplaceValue(MoveTemp(Result.GetErrorValue()));
 			return;
@@ -747,7 +747,7 @@ TFuture<TDefaultErrorResultInternal<TSharedRef<FLobbyDataEOS>>> FLobbyDataEOS::C
 			TDefaultErrorResultInternal<FLobbyMemberServiceSnapshot> LobbyMemberSnapshotResult = LobbyDetails->GetLobbyMemberSnapshot(MemberAccountId);
 			if (LobbyMemberSnapshotResult.IsError())
 			{
-				UE_LOG(LogTemp, Warning, TEXT("[FLobbyDataEOS::Create] FLobbyDetailsEOS::GetLobbyMemberSnapshot Failed: Lobby[%s], User[%s], Result[%s]"),
+				UE_LOG(LogOnlineServices, Warning, TEXT("[FLobbyDataEOS::Create] FLobbyDetailsEOS::GetLobbyMemberSnapshot Failed: Lobby[%s], User[%s], Result[%s]"),
 					*ToLogString(LobbyId), *ToLogString(MemberAccountId), *Result.GetErrorValue().GetLogString());
 				Promise.EmplaceValue(MoveTemp(LobbyMemberSnapshotResult.GetErrorValue()));
 				return;
@@ -761,7 +761,7 @@ TFuture<TDefaultErrorResultInternal<TSharedRef<FLobbyDataEOS>>> FLobbyDataEOS::C
 
 		if (PrepareServiceLobbySnapshotResult.IsError())
 		{
-			UE_LOG(LogTemp, Warning, TEXT("[FLobbyDataEOS::Create] FLobbyClientData::PrepareServiceSnapshot Failed: Lobby[%s], Result[%s]"),
+			UE_LOG(LogOnlineServices, Warning, TEXT("[FLobbyDataEOS::Create] FLobbyClientData::PrepareServiceSnapshot Failed: Lobby[%s], Result[%s]"),
 				*ToLogString(LobbyId), *PrepareServiceLobbySnapshotResult.GetErrorValue().GetLogString());
 			Promise.EmplaceValue(MoveTemp(PrepareServiceLobbySnapshotResult.GetErrorValue()));
 			return;
@@ -769,7 +769,7 @@ TFuture<TDefaultErrorResultInternal<TSharedRef<FLobbyDataEOS>>> FLobbyDataEOS::C
 
 		NewLobbyClientData->CommitServiceSnapshot({});
 
-		UE_LOG(LogTemp, VeryVerbose, TEXT("[FLobbyDataEOS::Create] Succeeded: Lobby[%s]"),
+		UE_LOG(LogOnlineServices, VeryVerbose, TEXT("[FLobbyDataEOS::Create] Succeeded: Lobby[%s]"),
 			UTF8_TO_TCHAR(LobbyDetails->GetInfo()->GetLobbyIdEOS()));
 		Promise.EmplaceValue(MakeShared<FLobbyDataEOS>(
 			NewLobbyClientData,
@@ -802,7 +802,7 @@ TFuture<TDefaultErrorResultInternal<TSharedRef<FLobbyDataEOS>>> FLobbyDataRegist
 	// Check if lobby data already exists.
 	if (TSharedPtr<FLobbyDataEOS> FindResult = Find(LobbyDetails->GetInfo()->GetLobbyIdEOS()))
 	{
-		UE_LOG(LogTemp, VeryVerbose, TEXT("[FLobbyDataRegistryEOS::FindOrCreateFromLobbyDetails] Succeeded: Lobby data already exists. User[%s], Lobby[%s]"),
+		UE_LOG(LogOnlineServices, VeryVerbose, TEXT("[FLobbyDataRegistryEOS::FindOrCreateFromLobbyDetails] Succeeded: Lobby data already exists. User[%s], Lobby[%s]"),
 			*ToLogString(LocalAccountId), *FindResult->GetLobbyIdString());
 
 		FindResult->AddUserLobbyDetails(LocalAccountId, LobbyDetails);
@@ -822,7 +822,7 @@ TFuture<TDefaultErrorResultInternal<TSharedRef<FLobbyDataEOS>>> FLobbyDataRegist
 		TSharedPtr<FLobbyDataRegistryEOS> StrongThis = WeakThis.Pin();
 		if (!StrongThis.IsValid())
 		{
-			UE_LOG(LogTemp, Warning, TEXT("[FLobbyDataRegistryEOS::FindOrCreateFromLobbyDetails] Failed: FLobbyDataRegistryEOS has been destroyed. User[%s], Lobby[%s], Result[%s]"),
+			UE_LOG(LogOnlineServices, Warning, TEXT("[FLobbyDataRegistryEOS::FindOrCreateFromLobbyDetails] Failed: FLobbyDataRegistryEOS has been destroyed. User[%s], Lobby[%s], Result[%s]"),
 				*ToLogString(LocalAccountId), UTF8_TO_TCHAR(LobbyDetails->GetInfo()->GetLobbyIdEOS()), *Result.GetErrorValue().GetLogString());
 			Promise.EmplaceValue(Errors::NotFound());
 			return;
@@ -830,13 +830,13 @@ TFuture<TDefaultErrorResultInternal<TSharedRef<FLobbyDataEOS>>> FLobbyDataRegist
 
 		if (Result.IsError())
 		{
-			UE_LOG(LogTemp, Warning, TEXT("[FLobbyDataRegistryEOS::FindOrCreateFromLobbyDetails] FLobbyDataEOS::Create Failed: User[%s], Lobby[%s], Result[%s]"),
+			UE_LOG(LogOnlineServices, Warning, TEXT("[FLobbyDataRegistryEOS::FindOrCreateFromLobbyDetails] FLobbyDataEOS::Create Failed: User[%s], Lobby[%s], Result[%s]"),
 				*ToLogString(LocalAccountId), UTF8_TO_TCHAR(LobbyDetails->GetInfo()->GetLobbyIdEOS()), *Result.GetErrorValue().GetLogString());
 			Promise.EmplaceValue(MoveTemp(Result.GetErrorValue()));
 			return;
 		}
 
-		UE_LOG(LogTemp, VeryVerbose, TEXT("[FLobbyDataRegistryEOS::FindOrCreateFromLobbyDetails] Succeeded: Created new lobby data. User[%s], Lobby[%s]"),
+		UE_LOG(LogOnlineServices, VeryVerbose, TEXT("[FLobbyDataRegistryEOS::FindOrCreateFromLobbyDetails] Succeeded: Created new lobby data. User[%s], Lobby[%s]"),
 			*ToLogString(LocalAccountId), *Result.GetOkValue()->GetLobbyIdString());
 
 		StrongThis->Register(Result.GetOkValue());
@@ -883,7 +883,7 @@ TFuture<TDefaultErrorResultInternal<TSharedRef<FLobbyInviteDataEOS>>> FLobbyInvi
 	TDefaultErrorResultInternal<TSharedRef<FLobbyDetailsEOS>> LobbyDetailsResult = FLobbyDetailsEOS::CreateFromInviteId(Prerequisites, LocalAccountId, InviteIdEOS);
 	if (LobbyDetailsResult.IsError())
 	{
-		UE_LOG(LogTemp, Warning, TEXT("[FLobbyInviteDataEOS::CreateFromInviteId] FLobbyDetailsEOS::CreateFromInviteId Failed: User[%s], InviteId[%s], Result[%s]"),
+		UE_LOG(LogOnlineServices, Warning, TEXT("[FLobbyInviteDataEOS::CreateFromInviteId] FLobbyDetailsEOS::CreateFromInviteId Failed: User[%s], InviteId[%s], Result[%s]"),
 			*ToLogString(LocalAccountId), UTF8_TO_TCHAR(InviteIdEOS), *LobbyDetailsResult.GetErrorValue().GetLogString());
 
 		return MakeFulfilledPromise<TDefaultErrorResultInternal<TSharedRef<FLobbyInviteDataEOS>>>(MoveTemp(LobbyDetailsResult.GetErrorValue())).GetFuture();
@@ -900,7 +900,7 @@ TFuture<TDefaultErrorResultInternal<TSharedRef<FLobbyInviteDataEOS>>> FLobbyInvi
 	{
 		if (Result.IsError())
 		{
-			UE_LOG(LogTemp, Warning, TEXT("[FLobbyInviteDataEOS::CreateFromInviteId] FLobbyDataEOS::FindOrCreateFromLobbyDetails Failed: User[%s], InviteId[%s], Lobby[%s], Result[%s]"),
+			UE_LOG(LogOnlineServices, Warning, TEXT("[FLobbyInviteDataEOS::CreateFromInviteId] FLobbyDataEOS::FindOrCreateFromLobbyDetails Failed: User[%s], InviteId[%s], Lobby[%s], Result[%s]"),
 				*ToLogString(LocalAccountId), UTF8_TO_TCHAR(InviteId->Get()), UTF8_TO_TCHAR(LobbyDetails->GetInfo()->GetLobbyIdEOS()), *Result.GetErrorValue().GetLogString());
 
 			Promise.EmplaceValue(MoveTemp(Result.GetErrorValue()));
@@ -911,14 +911,14 @@ TFuture<TDefaultErrorResultInternal<TSharedRef<FLobbyInviteDataEOS>>> FLobbyInvi
 		const FAccountId SenderAccountId = FindAccountId(Sender);
 		if (!SenderAccountId.IsValid())
 		{
-			UE_LOG(LogTemp, Warning, TEXT("[FLobbyInviteDataEOS::CreateFromInviteId] Failed: Unable to find account id for sender. User[%s], InviteId[%s], Lobby[%s]"),
+			UE_LOG(LogOnlineServices, Warning, TEXT("[FLobbyInviteDataEOS::CreateFromInviteId] Failed: Unable to find account id for sender. User[%s], InviteId[%s], Lobby[%s]"),
 				*ToLogString(LocalAccountId), UTF8_TO_TCHAR(InviteId->Get()), UTF8_TO_TCHAR(LobbyDetails->GetInfo()->GetLobbyIdEOS()));
 
 			Promise.EmplaceValue(Errors::InvalidState());
 			return;
 		}
 
-		UE_LOG(LogTemp, VeryVerbose, TEXT("[FLobbyInviteDataEOS::CreateFromInviteId] Succeeded: User[%s], Sender[%s], InviteId[%s], Lobby[%s]"),
+		UE_LOG(LogOnlineServices, VeryVerbose, TEXT("[FLobbyInviteDataEOS::CreateFromInviteId] Succeeded: User[%s], Sender[%s], InviteId[%s], Lobby[%s]"),
 			*ToLogString(LocalAccountId), *ToLogString(SenderAccountId), UTF8_TO_TCHAR(InviteId->Get()), *Result.GetOkValue()->GetLobbyIdString());
 		Promise.EmplaceValue(MakeShared<FLobbyInviteDataEOS>(InviteId, LocalAccountId, SenderAccountId, LobbyDetails, Result.GetOkValue()));
 	});
@@ -956,7 +956,7 @@ TFuture<TDefaultErrorResultInternal<TSharedRef<FLobbySearchEOS>>> FLobbySearchEO
 	EOS_EResult EOSResult = EOS_Lobby_CreateLobbySearch(Prerequisites->LobbyInterfaceHandle, &CreateLobbySearchOptions, &SearchHandle->Get());
 	if (EOSResult != EOS_EResult::EOS_Success)
 	{
-		UE_LOG(LogTemp, Warning, TEXT("[FLobbySearchEOS::Create] EOS_Lobby_CreateLobbySearch Failed: User[%s], Result[%s]"),
+		UE_LOG(LogOnlineServices, Warning, TEXT("[FLobbySearchEOS::Create] EOS_Lobby_CreateLobbySearch Failed: User[%s], Result[%s]"),
 			*ToLogString(Params.LocalAccountId), *LexToString(EOSResult));
 		return MakeFulfilledPromise<TDefaultErrorResultInternal<TSharedRef<FLobbySearchEOS>>>(Errors::FromEOSResult(EOSResult)).GetFuture();
 	}
@@ -966,7 +966,7 @@ TFuture<TDefaultErrorResultInternal<TSharedRef<FLobbySearchEOS>>> FLobbySearchEO
 		TSharedPtr<FLobbyDataEOS> LobbyData = LobbyRegistry->Find(*Params.LobbyId);
 		if (!LobbyData)
 		{
-			UE_LOG(LogTemp, Warning, TEXT("[FLobbySearchEOS::Create] Failed: Unable to find lobby data for lobby. User[%s], Lobby[%s]"),
+			UE_LOG(LogOnlineServices, Warning, TEXT("[FLobbySearchEOS::Create] Failed: Unable to find lobby data for lobby. User[%s], Lobby[%s]"),
 				*ToLogString(Params.LocalAccountId), *ToLogString(Params.LobbyId));
 			return MakeFulfilledPromise<TDefaultErrorResultInternal<TSharedRef<FLobbySearchEOS>>>(Errors::InvalidParams()).GetFuture();
 		}
@@ -979,7 +979,7 @@ TFuture<TDefaultErrorResultInternal<TSharedRef<FLobbySearchEOS>>> FLobbySearchEO
 		EOSResult = EOS_LobbySearch_SetLobbyId(SearchHandle->Get(), &SetLobbyIdOptions);
 		if (EOSResult != EOS_EResult::EOS_Success)
 		{
-			UE_LOG(LogTemp, Warning, TEXT("[FLobbySearchEOS::Create] EOS_LobbySearch_SetLobbyId Failed: User[%s], Lobby[%s], Result[%s]"),
+			UE_LOG(LogOnlineServices, Warning, TEXT("[FLobbySearchEOS::Create] EOS_LobbySearch_SetLobbyId Failed: User[%s], Lobby[%s], Result[%s]"),
 				*ToLogString(Params.LocalAccountId), *LobbyData->GetLobbyIdString(), *LexToString(EOSResult));
 			return MakeFulfilledPromise<TDefaultErrorResultInternal<TSharedRef<FLobbySearchEOS>>>(Errors::FromEOSResult(EOSResult)).GetFuture();
 		}
@@ -994,7 +994,7 @@ TFuture<TDefaultErrorResultInternal<TSharedRef<FLobbySearchEOS>>> FLobbySearchEO
 		EOSResult = EOS_LobbySearch_SetTargetUserId(SearchHandle->Get(), &SetTargetUserIdOptions);
 		if (EOSResult != EOS_EResult::EOS_Success)
 		{
-			UE_LOG(LogTemp, Warning, TEXT("[FLobbySearchEOS::Create] EOS_LobbySearch_SetTargetUserId Failed: User[%s], Target[%s], Result[%s]"),
+			UE_LOG(LogOnlineServices, Warning, TEXT("[FLobbySearchEOS::Create] EOS_LobbySearch_SetTargetUserId Failed: User[%s], Target[%s], Result[%s]"),
 				*ToLogString(Params.LocalAccountId), *ToLogString(*Params.TargetUser), *LexToString(EOSResult));
 			return MakeFulfilledPromise<TDefaultErrorResultInternal<TSharedRef<FLobbySearchEOS>>>(Errors::FromEOSResult(EOSResult)).GetFuture();
 		}
@@ -1021,7 +1021,7 @@ TFuture<TDefaultErrorResultInternal<TSharedRef<FLobbySearchEOS>>> FLobbySearchEO
 			EOSResult = EOS_LobbySearch_SetParameter(SearchHandle->Get(), &SetParameterOptions);
 			if (EOSResult != EOS_EResult::EOS_Success)
 			{
-				UE_LOG(LogTemp, Warning, TEXT("[FLobbySearchEOS::Create] EOS_LobbySearch_SetParameter Failed: User[%s], Key[%s], Result[%s]"),
+				UE_LOG(LogOnlineServices, Warning, TEXT("[FLobbySearchEOS::Create] EOS_LobbySearch_SetParameter Failed: User[%s], Key[%s], Result[%s]"),
 					*ToLogString(Params.LocalAccountId), UTF8_TO_TCHAR(EOS_LOBBY_SEARCH_BUCKET_ID), *LexToString(EOSResult));
 				return MakeFulfilledPromise<TDefaultErrorResultInternal<TSharedRef<FLobbySearchEOS>>>(Errors::FromEOSResult(EOSResult)).GetFuture();
 			}
@@ -1035,7 +1035,7 @@ TFuture<TDefaultErrorResultInternal<TSharedRef<FLobbySearchEOS>>> FLobbySearchEO
 			ESchemaServiceAttributeFlags SchemaServiceAttributeFlags = ESchemaServiceAttributeFlags::None;
 			if (!SchemaCategoryInstance.VerifyBaseAttributeData(Filter.AttributeName, Filter.ComparisonValue, SchemaServiceAttributeId, SchemaServiceAttributeFlags))
 			{
-				UE_LOG(LogTemp, Warning, TEXT("[FLobbySearchEOS::Create] VerifyBaseAttributeData failed: User[%s], Key[%s], Value[%s]"),
+				UE_LOG(LogOnlineServices, Warning, TEXT("[FLobbySearchEOS::Create] VerifyBaseAttributeData failed: User[%s], Key[%s], Value[%s]"),
 					*ToLogString(Params.LocalAccountId), *Filter.AttributeName.ToString().ToLower(), *ToLogString(Filter.ComparisonValue));
 				return MakeFulfilledPromise<TDefaultErrorResultInternal<TSharedRef<FLobbySearchEOS>>>(Errors::InvalidParams()).GetFuture();
 			}
@@ -1052,7 +1052,7 @@ TFuture<TDefaultErrorResultInternal<TSharedRef<FLobbySearchEOS>>> FLobbySearchEO
 			EOSResult = EOS_LobbySearch_SetParameter(SearchHandle->Get(), &SetParameterOptions);
 			if (EOSResult != EOS_EResult::EOS_Success)
 			{
-				UE_LOG(LogTemp, Warning, TEXT("[FLobbySearchEOS::Create] EOS_LobbySearch_SetParameter Failed: User[%s], Key[%s], Value[%s], Result[%s]"),
+				UE_LOG(LogOnlineServices, Warning, TEXT("[FLobbySearchEOS::Create] EOS_LobbySearch_SetParameter Failed: User[%s], Key[%s], Value[%s], Result[%s]"),
 					*ToLogString(Params.LocalAccountId), *Filter.AttributeName.ToString().ToLower(), *ToLogString(Filter.ComparisonValue), *LexToString(EOSResult));
 				return MakeFulfilledPromise<TDefaultErrorResultInternal<TSharedRef<FLobbySearchEOS>>>(Errors::FromEOSResult(EOSResult)).GetFuture();
 			}
@@ -1078,7 +1078,7 @@ TFuture<TDefaultErrorResultInternal<TSharedRef<FLobbySearchEOS>>> FLobbySearchEO
 	{
 		if (Data->ResultCode != EOS_EResult::EOS_Success)
 		{
-			UE_LOG(LogTemp, Warning, TEXT("[FLobbySearchEOS::Create] EOS_LobbySearch_Find Failed: User[%s], Result[%s]"),
+			UE_LOG(LogOnlineServices, Warning, TEXT("[FLobbySearchEOS::Create] EOS_LobbySearch_Find Failed: User[%s], Result[%s]"),
 				*ToLogString(LocalAccountId), *LexToString(Data->ResultCode));
 			Promise.EmplaceValue(Errors::FromEOSResult(Data->ResultCode));
 			return;
@@ -1098,7 +1098,7 @@ TFuture<TDefaultErrorResultInternal<TSharedRef<FLobbySearchEOS>>> FLobbySearchEO
 				Prerequisites, LocalAccountId, SearchHandle->Get(), SearchResultIndex);
 			if (Result.IsError())
 			{
-				UE_LOG(LogTemp, Warning, TEXT("[FLobbySearchEOS::Create] FLobbyDetailsEOS::CreateFromSearchResult Failed: User[%s], Result[%s]"),
+				UE_LOG(LogOnlineServices, Warning, TEXT("[FLobbySearchEOS::Create] FLobbyDetailsEOS::CreateFromSearchResult Failed: User[%s], Result[%s]"),
 					*ToLogString(LocalAccountId), *Result.GetErrorValue().GetLogString());
 				Promise.EmplaceValue(MoveTemp(Result.GetErrorValue()));
 				return;
@@ -1124,7 +1124,7 @@ TFuture<TDefaultErrorResultInternal<TSharedRef<FLobbySearchEOS>>> FLobbySearchEO
 			{
 				if (Result.IsError())
 				{
-					UE_LOG(LogTemp, Warning, TEXT("[FLobbySearchEOS::Create] FLobbyDataRegistryEOS::FindOrCreateFromLobbyDetails Failed: Unable to resolve search result. User[%s], Result[%s]"),
+					UE_LOG(LogOnlineServices, Warning, TEXT("[FLobbySearchEOS::Create] FLobbyDataRegistryEOS::FindOrCreateFromLobbyDetails Failed: Unable to resolve search result. User[%s], Result[%s]"),
 						*ToLogString(LocalAccountId), *Result.GetErrorValue().GetLogString());
 					Promise.EmplaceValue(MoveTemp(Result.GetErrorValue()));
 					return;
@@ -1133,7 +1133,7 @@ TFuture<TDefaultErrorResultInternal<TSharedRef<FLobbySearchEOS>>> FLobbySearchEO
 				ResolvedResults.Add(MoveTemp(Result.GetOkValue()));
 			}
 
-			UE_LOG(LogTemp, VeryVerbose, TEXT("[FLobbySearchEOS::Create] Succeeded: User[%s], NumResults[%d]"),
+			UE_LOG(LogOnlineServices, VeryVerbose, TEXT("[FLobbySearchEOS::Create] Succeeded: User[%s], NumResults[%d]"),
 				*ToLogString(LocalAccountId), ResolvedResults.Num());
 			Promise.EmplaceValue(MakeShared<FLobbySearchEOS>(SearchHandle, MoveTemp(ResolvedResults)));
 		});
diff --git a/Engine/Plugins/Online/OnlineServicesEOSGS/Source/Private/Online/OnlineServicesEOSGS.cpp b/Engine/Plugins/Online/OnlineServicesEOSGS/Source/Private/Online/OnlineServicesEOSGS.cpp
--- a/Engine/Plugins/Online/OnlineServicesEOSGS/Source/Private/Online/OnlineServicesEOSGS.cpp
+++ b/Engine/Plugins/Online/OnlineServicesEOSGS/Source/Private/Online/OnlineServicesEOSGS.cpp
@@ -12,6 +12,7 @@
 #include "Online/SessionsEOSGS.h"
 #include "Online/TitleFileEOSGS.h"
 #include "Online/UserFileEOSGS.h"
+#include "IEOSSDKManager.h"
 
 #if WITH_ENGINE
 #include "InternetAddrEOS.h"
diff --git a/Engine/Plugins/Online/OnlineServicesEOSGS/Source/Private/Online/SessionsEOSGS.cpp b/Engine/Plugins/Online/OnlineServicesEOSGS/Source/Private/Online/SessionsEOSGS.cpp
--- a/Engine/Plugins/Online/OnlineServicesEOSGS/Source/Private/Online/SessionsEOSGS.cpp
+++ b/Engine/Plugins/Online/OnlineServicesEOSGS/Source/Private/Online/SessionsEOSGS.cpp
@@ -115,13 +115,13 @@ FSessionEOSGS::FSessionEOSGS(const TSharedPtr<FSessionDetailsHandleEOSGS>& InSes
 			}
 			else
 			{
-				UE_LOG(LogTemp, Warning, TEXT("[FSessionEOSGS] EOS_SessionDetails_CopySessionAttributeByIndex failed with result [%s]"), *LexToString(CopyInfoResult));
+				UE_LOG(LogOnlineServices, Warning, TEXT("[FSessionEOSGS] EOS_SessionDetails_CopySessionAttributeByIndex failed with result [%s]"), *LexToString(CopyInfoResult));
 			}
 		}
 	}
 	else
 	{
-		UE_LOG(LogTemp, Warning, TEXT("[FSessionEOSGS] EOS_SessionDetails_CopyInfo failed with result [%s]"), *LexToString(CopyInfoResult));
+		UE_LOG(LogOnlineServices, Warning, TEXT("[FSessionEOSGS] EOS_SessionDetails_CopyInfo failed with result [%s]"), *LexToString(CopyInfoResult));
 	}
 }
 
@@ -340,7 +340,7 @@ TFuture<TOnlineResult<FCreateSession>> FSessionsEOSGS::CreateSessionImpl(const F
 	{
 		int32 BuildUniqueId = GetBuildUniqueId();
 
-		UE_LOG(LogTemp, Verbose, TEXT("[FSessionsEOSGS::CreateSession] 'EOSGS_BUCKET_ID_ATTRIBUTE_KEY' (FString) Custom Setting needed to create EOS sessions not found. Setting \"%d\" as default."), BuildUniqueId);
+		UE_LOG(LogOnlineServices, Verbose, TEXT("[FSessionsEOSGS::CreateSession] 'EOSGS_BUCKET_ID_ATTRIBUTE_KEY' (FString) Custom Setting needed to create EOS sessions not found. Setting \"%d\" as default."), BuildUniqueId);
 
 		BucketIdStr = FString::FromInt(BuildUniqueId);
 	}
@@ -364,7 +364,7 @@ TFuture<TOnlineResult<FCreateSession>> FSessionsEOSGS::CreateSessionImpl(const F
 	EOS_EResult ResultCode = EOS_Sessions_CreateSessionModification(SessionsHandle, &CreateSessionModificationOptions, &SessionModificationHandle);
 	if (ResultCode != EOS_EResult::EOS_Success)
 	{
-		UE_LOG(LogTemp, Error, TEXT("[FSessionsEOSGS::CreateSession] EOS_Sessions_CreateSessionModification failed with result [%s]"), *LexToString(ResultCode));
+		UE_LOG(LogOnlineServices, Error, TEXT("[FSessionsEOSGS::CreateSession] EOS_Sessions_CreateSessionModification failed with result [%s]"), *LexToString(ResultCode));
 
 		Promise.EmplaceValue(Errors::FromEOSResult(ResultCode));
 
@@ -419,7 +419,7 @@ TOptional<FOnlineError> FSessionsEOSGS::CheckState(const FCreateSession::Params&
 
 		if (Length < EOS_SESSIONMODIFICATION_MIN_SESSIONIDOVERRIDE_LENGTH || Length > EOS_SESSIONMODIFICATION_MAX_SESSIONIDOVERRIDE_LENGTH)
 		{
-			UE_LOG(LogTemp, Warning, TEXT("[%s] Could not create session with SessionIdOverride [%s] of size [%d]. SessionIdOverride size must be between [%d] and [%d] characters long"), UTF8_TO_TCHAR(__FUNCTION__), *Params.SessionIdOverride, Length, EOS_SESSIONMODIFICATION_MIN_SESSIONIDOVERRIDE_LENGTH, EOS_SESSIONMODIFICATION_MAX_SESSIONIDOVERRIDE_LENGTH);
+			UE_LOG(LogOnlineServices, Warning, TEXT("[%s] Could not create session with SessionIdOverride [%s] of size [%d]. SessionIdOverride size must be between [%d] and [%d] characters long"), UTF8_TO_TCHAR(__FUNCTION__), *Params.SessionIdOverride, Length, EOS_SESSIONMODIFICATION_MIN_SESSIONIDOVERRIDE_LENGTH, EOS_SESSIONMODIFICATION_MAX_SESSIONIDOVERRIDE_LENGTH);
 
 			return TOptional<FOnlineError>(Errors::InvalidParams());
 		}
@@ -440,7 +440,7 @@ void FSessionsEOSGS::SetHostAddress(EOS_HSessionModification& SessionModHandle,
 	EOS_EResult ResultCode = EOS_SessionModification_SetHostAddress(SessionModHandle, &Options);
 	if (ResultCode != EOS_EResult::EOS_Success)
 	{
-		UE_LOG(LogTemp, Warning, TEXT("EOS_SessionModification_SetHostAddress failed with result [%s]"), *LexToString(ResultCode));
+		UE_LOG(LogOnlineServices, Warning, TEXT("EOS_SessionModification_SetHostAddress failed with result [%s]"), *LexToString(ResultCode));
 	}
 }
 
@@ -455,7 +455,7 @@ void FSessionsEOSGS::SetJoinInProgressAllowed(EOS_HSessionModification& SessionM
 	EOS_EResult ResultCode = EOS_SessionModification_SetJoinInProgressAllowed(SessionModHandle, &Options);
 	if (ResultCode != EOS_EResult::EOS_Success)
 	{
-		UE_LOG(LogTemp, Warning, TEXT("EOS_SessionModification_SetJoinInProgressAllowed failed with result [%s]"), *LexToString(ResultCode));
+		UE_LOG(LogOnlineServices, Warning, TEXT("EOS_SessionModification_SetJoinInProgressAllowed failed with result [%s]"), *LexToString(ResultCode));
 	}
 }
 
@@ -470,7 +470,7 @@ void FSessionsEOSGS::SetInvitesAllowed(EOS_HSessionModification& SessionModHandl
 	EOS_EResult ResultCode = EOS_SessionModification_SetInvitesAllowed(SessionModHandle, &Options);
 	if (ResultCode != EOS_EResult::EOS_Success)
 	{
-		UE_LOG(LogTemp, Warning, TEXT("EOS_SessionModification_SetInvitesAllowed failed with result [%s]"), *LexToString(ResultCode));
+		UE_LOG(LogOnlineServices, Warning, TEXT("EOS_SessionModification_SetInvitesAllowed failed with result [%s]"), *LexToString(ResultCode));
 	}
 }
 
@@ -485,7 +485,7 @@ void FSessionsEOSGS::SetPermissionLevel(EOS_HSessionModification& SessionModific
 	EOS_EResult ResultCode = EOS_SessionModification_SetPermissionLevel(SessionModificationHandle, &Options);
 	if (ResultCode != EOS_EResult::EOS_Success)
 	{
-		UE_LOG(LogTemp, Warning, TEXT("[FSessionsEOSGS::SetPermissionLevel] EOS_SessionModification_SetPermissionLevel failed with result [%s]"), *LexToString(ResultCode));
+		UE_LOG(LogOnlineServices, Warning, TEXT("[FSessionsEOSGS::SetPermissionLevel] EOS_SessionModification_SetPermissionLevel failed with result [%s]"), *LexToString(ResultCode));
 	}
 }
 
@@ -501,7 +501,7 @@ void FSessionsEOSGS::SetBucketId(EOS_HSessionModification& SessionModificationHa
 	EOS_EResult ResultCode = EOS_SessionModification_SetBucketId(SessionModificationHandle, &Options);
 	if (ResultCode != EOS_EResult::EOS_Success)
 	{
-		UE_LOG(LogTemp, Warning, TEXT("[FSessionsEOSGS::SetBucketId] EOS_SessionModification_SetBucketId failed with result [%s]"), *LexToString(ResultCode));
+		UE_LOG(LogOnlineServices, Warning, TEXT("[FSessionsEOSGS::SetBucketId] EOS_SessionModification_SetBucketId failed with result [%s]"), *LexToString(ResultCode));
 	}
 }
 
@@ -516,7 +516,7 @@ void FSessionsEOSGS::SetMaxPlayers(EOS_HSessionModification& SessionModification
 	EOS_EResult ResultCode = EOS_SessionModification_SetMaxPlayers(SessionModificationHandle, &Options);
 	if (ResultCode != EOS_EResult::EOS_Success)
 	{
-		UE_LOG(LogTemp, Warning, TEXT("[FSessionsEOSGS::SetMaxPlayers] EOS_SessionModification_SetMaxPlayers failed with result [%s]"), *LexToString(ResultCode));
+		UE_LOG(LogOnlineServices, Warning, TEXT("[FSessionsEOSGS::SetMaxPlayers] EOS_SessionModification_SetMaxPlayers failed with result [%s]"), *LexToString(ResultCode));
 	}
 }
 
@@ -534,7 +534,7 @@ void FSessionsEOSGS::AddAttribute(EOS_HSessionModification& SessionModificationH
 	EOS_EResult ResultCode = EOS_SessionModification_AddAttribute(SessionModificationHandle, &Options);
 	if (ResultCode != EOS_EResult::EOS_Success)
 	{
-		UE_LOG(LogTemp, Warning, TEXT("[FSessionsEOSGS::AddAttribute] EOS_SessionModification_AddAttribute failed with result [%s]"), *LexToString(ResultCode));
+		UE_LOG(LogOnlineServices, Warning, TEXT("[FSessionsEOSGS::AddAttribute] EOS_SessionModification_AddAttribute failed with result [%s]"), *LexToString(ResultCode));
 	}
 }
 
@@ -550,7 +550,7 @@ void FSessionsEOSGS::RemoveAttribute(EOS_HSessionModification& SessionModificati
 	EOS_EResult ResultCode = EOS_SessionModification_RemoveAttribute(SessionModificationHandle, &Options);
 	if (ResultCode != EOS_EResult::EOS_Success)
 	{
-		UE_LOG(LogTemp, Warning, TEXT("[FSessionsEOSGS::RemoveAttribute] EOS_SessionModification_RemoveAttribute failed with result [%s]"), *LexToString(ResultCode));
+		UE_LOG(LogOnlineServices, Warning, TEXT("[FSessionsEOSGS::RemoveAttribute] EOS_SessionModification_RemoveAttribute failed with result [%s]"), *LexToString(ResultCode));
 	}
 }
 
@@ -679,7 +679,7 @@ TFuture<TOnlineResult<FUpdateSessionSettings>> FSessionsEOSGS::UpdateSessionSett
 	EOS_EResult ResultCode = EOS_Sessions_UpdateSessionModification(SessionsHandle, &UpdateSessionModificationOptions, &SessionModificationHandle);
 	if (ResultCode != EOS_EResult::EOS_Success)
 	{
-		UE_LOG(LogTemp, Error, TEXT("[FSessionsEOSGS::UpdateSession] EOS_Sessions_UpdateSessionModification failed with result [%s]"), *LexToString(ResultCode));
+		UE_LOG(LogOnlineServices, Error, TEXT("[FSessionsEOSGS::UpdateSession] EOS_Sessions_UpdateSessionModification failed with result [%s]"), *LexToString(ResultCode));
 
 		Promise.EmplaceValue(Errors::FromEOSResult(ResultCode));
 
@@ -750,7 +750,7 @@ TFuture<TDefaultErrorResult<FUpdateSessionImplEOSGS>> FSessionsEOSGS::UpdateSess
 		// If we only change bAllowNewMembers, the session update will yield an EOS_NoChange result, but we still need to continue to the next step
 		if (Result->ResultCode != EOS_EResult::EOS_Success && Result->ResultCode != EOS_EResult::EOS_Sessions_OutOfSync && Result->ResultCode != EOS_EResult::EOS_NoChange)
 		{
-			UE_LOG(LogTemp, Warning, TEXT("EOS_Sessions_UpdateSession failed with result [%s]"), *LexToString(Result->ResultCode));
+			UE_LOG(LogOnlineServices, Warning, TEXT("EOS_Sessions_UpdateSession failed with result [%s]"), *LexToString(Result->ResultCode));
 			Promise.EmplaceValue(Errors::FromEOSResult(Result->ResultCode));
 			return;
 		}
@@ -797,7 +797,7 @@ TFuture<TDefaultErrorResult<FUpdateSessionJoinabilityImpl>> FSessionsEOSGS::Upda
 	EOS_EResult CopyActiveSessionHandleResult = EOS_Sessions_CopyActiveSessionHandle(SessionsHandle, &CopyActiveSessionHandleOptions, &ActiveSessionHandle);
 	if (CopyActiveSessionHandleResult != EOS_EResult::EOS_Success)
 	{
-		UE_LOG(LogTemp, Warning, TEXT("EOS_Sessions_CopyActiveSessionHandle failed with result [%s]"), *LexToString(CopyActiveSessionHandleResult));
+		UE_LOG(LogOnlineServices, Warning, TEXT("EOS_Sessions_CopyActiveSessionHandle failed with result [%s]"), *LexToString(CopyActiveSessionHandleResult));
 	}
 
 	// We get the active session info with the handle
@@ -809,7 +809,7 @@ TFuture<TDefaultErrorResult<FUpdateSessionJoinabilityImpl>> FSessionsEOSGS::Upda
 	EOS_EResult CopyInfoResult = EOS_ActiveSession_CopyInfo(ActiveSessionHandle, &CopyInfoOptions, &ActiveSessionInfo);
 	if (CopyInfoResult != EOS_EResult::EOS_Success)
 	{
-		UE_LOG(LogTemp, Warning, TEXT("EOS_ActiveSession_CopyInfo failed with result [%s]"), *LexToString(CopyInfoResult));
+		UE_LOG(LogOnlineServices, Warning, TEXT("EOS_ActiveSession_CopyInfo failed with result [%s]"), *LexToString(CopyInfoResult));
 	}
 
 	// If not, we start or end the session to make it joinable or not (as we set JIP to false at creation time)
@@ -829,7 +829,7 @@ TFuture<TDefaultErrorResult<FUpdateSessionJoinabilityImpl>> FSessionsEOSGS::Upda
 			{
 				if (Result->ResultCode != EOS_EResult::EOS_Success && Result->ResultCode != EOS_EResult::EOS_Sessions_OutOfSync)
 				{
-					UE_LOG(LogTemp, Warning, TEXT("EOS_Sessions_EndSession failed with result [%s]"), *LexToString(Result->ResultCode));
+					UE_LOG(LogOnlineServices, Warning, TEXT("EOS_Sessions_EndSession failed with result [%s]"), *LexToString(Result->ResultCode));
 					Promise.EmplaceValue(Errors::FromEOSResult(Result->ResultCode));
 					return;
 				}
@@ -858,7 +858,7 @@ TFuture<TDefaultErrorResult<FUpdateSessionJoinabilityImpl>> FSessionsEOSGS::Upda
 			{
 				if (Result->ResultCode != EOS_EResult::EOS_Success && Result->ResultCode != EOS_EResult::EOS_Sessions_OutOfSync)
 				{
-					UE_LOG(LogTemp, Warning, TEXT("EOS_Sessions_StartSession failed with result [%s]"), *LexToString(Result->ResultCode));
+					UE_LOG(LogOnlineServices, Warning, TEXT("EOS_Sessions_StartSession failed with result [%s]"), *LexToString(Result->ResultCode));
 					Promise.EmplaceValue(Errors::FromEOSResult(Result->ResultCode));
 					return;
 				}
@@ -896,7 +896,7 @@ TFuture<TDefaultErrorResult<FSendSingleSessionInviteImpl>> FSessionsEOSGS::SendS
 	{
 		if (Result->ResultCode != EOS_EResult::EOS_Success)
 		{
-			UE_LOG(LogTemp, Warning, TEXT("EOS_Sessions_SendInvite failed with result [%s]"), *LexToString(Result->ResultCode));
+			UE_LOG(LogOnlineServices, Warning, TEXT("EOS_Sessions_SendInvite failed with result [%s]"), *LexToString(Result->ResultCode));
 			Promise.EmplaceValue(Errors::FromEOSResult(Result->ResultCode));
 		}
 		else
@@ -934,7 +934,7 @@ TFuture<TOnlineResult<FLeaveSession>> FSessionsEOSGS::LeaveSessionImpl(const FLe
 		{
 			if (Result->ResultCode != EOS_EResult::EOS_Success)
 			{
-				UE_LOG(LogTemp, Warning, TEXT("EOS_Sessions_DestroySession failed with result [%s]"), *LexToString(Result->ResultCode));
+				UE_LOG(LogOnlineServices, Warning, TEXT("EOS_Sessions_DestroySession failed with result [%s]"), *LexToString(Result->ResultCode));
 				Promise.EmplaceValue(Errors::FromEOSResult(Result->ResultCode));
 				return;
 			}
@@ -973,7 +973,7 @@ void FSessionsEOSGS::SetSessionSearchMaxResults(FSessionSearchHandleEOSGS& Sessi
 	EOS_EResult ResultCode = EOS_SessionSearch_SetMaxResults(SessionSearchHandle.SearchHandle, &Options);
 	if (ResultCode != EOS_EResult::EOS_Success)
 	{
-		UE_LOG(LogTemp, Warning, TEXT("[FSessionsEOSGS::SetSessionSearchMaxResults] EOS_SessionSearch_SetMaxResults failed with result [%s]"), *LexToString(ResultCode));
+		UE_LOG(LogOnlineServices, Warning, TEXT("[FSessionsEOSGS::SetSessionSearchMaxResults] EOS_SessionSearch_SetMaxResults failed with result [%s]"), *LexToString(ResultCode));
 	}
 }
 
@@ -993,7 +993,7 @@ void FSessionsEOSGS::SetSessionSearchParameters(FSessionSearchHandleEOSGS& Sessi
 		EOS_EResult ResultCode = EOS_SessionSearch_SetParameter(SessionSearchHandle.SearchHandle, &Options);
 		if (ResultCode != EOS_EResult::EOS_Success)
 		{
-			UE_LOG(LogTemp, Warning, TEXT("[FSessionsEOSGS::SetSessionSearchParameters] EOS_SessionSearch_SetParameter failed with result [%s]"), *LexToString(ResultCode));
+			UE_LOG(LogOnlineServices, Warning, TEXT("[FSessionsEOSGS::SetSessionSearchParameters] EOS_SessionSearch_SetParameter failed with result [%s]"), *LexToString(ResultCode));
 		}
 	}
 }
@@ -1010,7 +1010,7 @@ void FSessionsEOSGS::SetSessionSearchSessionId(FSessionSearchHandleEOSGS& Sessio
 	EOS_EResult ResultCode = EOS_SessionSearch_SetSessionId(SessionSearchHandle.SearchHandle, &Options);
 	if (ResultCode != EOS_EResult::EOS_Success)
 	{
-		UE_LOG(LogTemp, Warning, TEXT("[FSessionsEOSGS::SetSessionSearchSessionId] EOS_SessionSearch_SetSessionId failed with result [%s]"), *LexToString(ResultCode));
+		UE_LOG(LogOnlineServices, Warning, TEXT("[FSessionsEOSGS::SetSessionSearchSessionId] EOS_SessionSearch_SetSessionId failed with result [%s]"), *LexToString(ResultCode));
 	}
 }
 
@@ -1024,7 +1024,7 @@ void FSessionsEOSGS::SetSessionSearchTargetId(FSessionSearchHandleEOSGS& Session
 	EOS_EResult ResultCode = EOS_SessionSearch_SetTargetUserId(SessionSearchHandle.SearchHandle, &Options);
 	if (ResultCode != EOS_EResult::EOS_Success)
 	{
-		UE_LOG(LogTemp, Warning, TEXT("[FSessionsEOSGS::SetSessionSearchTargetId] EOS_SessionSearch_SetTargetUserId failed with result [%s]"), *LexToString(ResultCode));
+		UE_LOG(LogOnlineServices, Warning, TEXT("[FSessionsEOSGS::SetSessionSearchTargetId] EOS_SessionSearch_SetTargetUserId failed with result [%s]"), *LexToString(ResultCode));
 	}
 }
 
@@ -1080,7 +1080,7 @@ TFuture<TOnlineResult<FFindSessions>> FSessionsEOSGS::FindSessionsImpl(const FFi
 	EOS_EResult ResultCode = EOS_Sessions_CreateSessionSearch(SessionsHandle, &CreateSessionSearchOptions, &SearchHandle);
 	if (ResultCode != EOS_EResult::EOS_Success)
 	{
-		UE_LOG(LogTemp, Error, TEXT("[FSessionsEOSGS::FindSessions] EOS_Sessions_CreateSessionSearch failed with result [%s]"), *LexToString(ResultCode));
+		UE_LOG(LogOnlineServices, Error, TEXT("[FSessionsEOSGS::FindSessions] EOS_Sessions_CreateSessionSearch failed with result [%s]"), *LexToString(ResultCode));
 
 		Promise.EmplaceValue(Errors::FromEOSResult(ResultCode));
 		return Future;
@@ -1107,7 +1107,7 @@ TFuture<TOnlineResult<FFindSessions>> FSessionsEOSGS::FindSessionsImpl(const FFi
 
 			if (FindCallbackInfoResult->ResultCode != EOS_EResult::EOS_Success)
 			{
-				UE_LOG(LogTemp, Warning, TEXT("EOS_SessionSearch_Find failed with result [%s]"), *LexToString(FindCallbackInfoResult->ResultCode));
+				UE_LOG(LogOnlineServices, Warning, TEXT("EOS_SessionSearch_Find failed with result [%s]"), *LexToString(FindCallbackInfoResult->ResultCode));
 				Promise.EmplaceValue(Errors::FromEOSResult(FindCallbackInfoResult->ResultCode));
 				CurrentSessionSearchPromisesUserMap.Remove(Params.LocalAccountId);
 				CurrentSessionSearchHandleEOSGSUserMap.Remove(Params.LocalAccountId);
@@ -1156,7 +1156,7 @@ TFuture<TOnlineResult<FFindSessions>> FSessionsEOSGS::FindSessionsImpl(const FFi
 				}
 				else
 				{
-					UE_LOG(LogTemp, Error, TEXT("[FSessionsEOSGS::FindSessions] EOS_SessionSearch_CopySearchResultByIndex failed for index [%d] with result [%s]"), Index, *LexToString(CopySearchResultByIndexResult));
+					UE_LOG(LogOnlineServices, Error, TEXT("[FSessionsEOSGS::FindSessions] EOS_SessionSearch_CopySearchResultByIndex failed for index [%d] with result [%s]"), Index, *LexToString(CopySearchResultByIndexResult));
 
 					Promise.EmplaceValue(Errors::FromEOSResult(CopySearchResultByIndexResult));
 					return;
@@ -1202,7 +1202,7 @@ TFuture<TOnlineResult<FJoinSession>> FSessionsEOSGS::JoinSessionImpl(const FJoin
 		// If no result is found, the id might be expired, which we should notify
 		if (FOnlineSessionIdRegistryEOSGS::Get().IsSessionIdExpired(Params.SessionId))
 		{
-			UE_LOG(LogTemp, Warning, TEXT("[%s] SessionId parameter [%s] is expired. Please call FindSessions to get an updated list of available sessions "), UTF8_TO_TCHAR(__FUNCTION__), *ToLogString(Params.SessionId));
+			UE_LOG(LogOnlineServices, Warning, TEXT("[%s] SessionId parameter [%s] is expired. Please call FindSessions to get an updated list of available sessions "), UTF8_TO_TCHAR(__FUNCTION__), *ToLogString(Params.SessionId));
 		}
 
 		Promise.EmplaceValue(GetSessionByIdResult.GetErrorValue());
@@ -1239,7 +1239,7 @@ TFuture<TOnlineResult<FJoinSession>> FSessionsEOSGS::JoinSessionImpl(const FJoin
 		{
 			if (Result->ResultCode != EOS_EResult::EOS_Success)
 			{
-				UE_LOG(LogTemp, Warning, TEXT("EOS_Sessions_JoinSession failed with result [%s]"), *LexToString(Result->ResultCode));
+				UE_LOG(LogOnlineServices, Warning, TEXT("EOS_Sessions_JoinSession failed with result [%s]"), *LexToString(Result->ResultCode));
 				Promise.EmplaceValue(Errors::FromEOSResult(Result->ResultCode));
 				return;
 			}
@@ -1250,7 +1250,7 @@ TFuture<TOnlineResult<FJoinSession>> FSessionsEOSGS::JoinSessionImpl(const FJoin
 				// If no result is found, the id might be expired, which we should notify
 				if (FOnlineSessionIdRegistryEOSGS::Get().IsSessionIdExpired(Params.SessionId))
 				{
-					UE_LOG(LogTemp, Warning, TEXT("[%s] SessionId parameter [%s] is expired. Please call FindSessions to get an updated list of available sessions "), UTF8_TO_TCHAR(__FUNCTION__), *ToLogString(Params.SessionId));
+					UE_LOG(LogOnlineServices, Warning, TEXT("[%s] SessionId parameter [%s] is expired. Please call FindSessions to get an updated list of available sessions "), UTF8_TO_TCHAR(__FUNCTION__), *ToLogString(Params.SessionId));
 				}
 
 				Promise.EmplaceValue(MoveTemp(GetSessionByIdResult.GetErrorValue()));
@@ -1292,7 +1292,7 @@ TOptional<FOnlineError> FSessionsEOSGS::CheckState(const FJoinSession::Params& P
 		const FSessionEOSGS& SessionEOSGS = FSessionEOSGS::Cast(*FoundSession);
 		if (!FoundSession->GetSessionInfo().bIsLANSession && !SessionEOSGS.SessionDetailsHandle.IsValid())
 		{
-			UE_LOG(LogTemp, Warning, TEXT("[%s] Could not join session with invalid session details handle in session with id [%s]"), UTF8_TO_TCHAR(__FUNCTION__), *ToLogString(Params.SessionId));
+			UE_LOG(LogOnlineServices, Warning, TEXT("[%s] Could not join session with invalid session details handle in session with id [%s]"), UTF8_TO_TCHAR(__FUNCTION__), *ToLogString(Params.SessionId));
 
 			return TOptional<FOnlineError>(Errors::InvalidState());
 		}
@@ -1322,7 +1322,7 @@ TOnlineAsyncOpHandle<FBuildSessionFromDetailsHandle> FSessionsEOSGS::BuildSessio
 	}
 	else
 	{
-		UE_LOG(LogTemp, Error, TEXT("[FSessionsEOSGS::BuildSessionFromInvite] EOS_Sessions_CopySessionHandleByInviteId failed with result [%s]"), *LexToString(CopySessionHandleByInviteIdResult));
+		UE_LOG(LogOnlineServices, Error, TEXT("[FSessionsEOSGS::BuildSessionFromInvite] EOS_Sessions_CopySessionHandleByInviteId failed with result [%s]"), *LexToString(CopySessionHandleByInviteIdResult));
 
 		TOnlineAsyncOpPtr<FBuildSessionFromDetailsHandle> Operation;
 		Operation->SetError(Errors::FromEOSResult(CopySessionHandleByInviteIdResult));
@@ -1346,7 +1346,7 @@ TOnlineAsyncOpHandle<FBuildSessionFromDetailsHandle> FSessionsEOSGS::BuildSessio
 	}
 	else
 	{
-		UE_LOG(LogTemp, Error, TEXT("[FSessionsEOSGS::BuildSessionFromUIEvent] EOS_Sessions_CopySessionHandleByUiEventId failed with result [%s]"), *LexToString(CopySessionHandleByUiEventIdResult));
+		UE_LOG(LogOnlineServices, Error, TEXT("[FSessionsEOSGS::BuildSessionFromUIEvent] EOS_Sessions_CopySessionHandleByUiEventId failed with result [%s]"), *LexToString(CopySessionHandleByUiEventIdResult));
 
 		TOnlineAsyncOpPtr<FBuildSessionFromDetailsHandle> Operation;
 		Operation->SetError(Errors::FromEOSResult(CopySessionHandleByUiEventIdResult));
@@ -1372,7 +1372,7 @@ TResult<TArray<EOS_ProductUserId>, FOnlineError> GetProductUserIdsFromEOSGSSessi
 	}
 	else
 	{
-		UE_LOG(LogTemp, Warning, TEXT("[FSessionsEOSGS::BuildSessionFromDetailsHandle] EOS_SessionDetails_CopyInfo failed with result [%s]"), *LexToString(CopyInfoResult));
+		UE_LOG(LogOnlineServices, Warning, TEXT("[FSessionsEOSGS::BuildSessionFromDetailsHandle] EOS_SessionDetails_CopyInfo failed with result [%s]"), *LexToString(CopyInfoResult));
 
 		TResult<TArray<EOS_ProductUserId>, FOnlineError>(Errors::FromEOSResult(CopyInfoResult));
 	}*/
@@ -1476,7 +1476,7 @@ TFuture<TOnlineResult<FRejectSessionInvite>> FSessionsEOSGS::RejectSessionInvite
 		{
 			if (Result->ResultCode != EOS_EResult::EOS_Success)
 			{
-				UE_LOG(LogTemp, Warning, TEXT("EOS_Sessions_RejectInvite failed with result [%s]"), *LexToString(Result->ResultCode));
+				UE_LOG(LogOnlineServices, Warning, TEXT("EOS_Sessions_RejectInvite failed with result [%s]"), *LexToString(Result->ResultCode));
 				Promise.EmplaceValue(Errors::FromEOSResult(Result->ResultCode));
 				return;
 			}
@@ -1528,7 +1528,7 @@ TFuture<TOnlineResult<FAddSessionMember>> FSessionsEOSGS::AddSessionMemberImpl(c
 						{
 							if (Result->ResultCode != EOS_EResult::EOS_Success)
 							{
-								UE_LOG(LogTemp, Warning, TEXT("EOS_Sessions_RegisterPlayers failed with result [%s]"), *LexToString(Result->ResultCode));
+								UE_LOG(LogOnlineServices, Warning, TEXT("EOS_Sessions_RegisterPlayers failed with result [%s]"), *LexToString(Result->ResultCode));
 								Promise.EmplaceValue(Errors::FromEOSResult(Result->ResultCode));
 								return;
 							}
@@ -1581,7 +1581,7 @@ TFuture<TOnlineResult<FRemoveSessionMember>> FSessionsEOSGS::RemoveSessionMember
 						{
 							if (Result->ResultCode != EOS_EResult::EOS_Success)
 							{
-								UE_LOG(LogTemp, Warning, TEXT("EOS_Sessions_UnregisterPlayers failed with result [%s]"), *LexToString(Result->ResultCode));
+								UE_LOG(LogOnlineServices, Warning, TEXT("EOS_Sessions_UnregisterPlayers failed with result [%s]"), *LexToString(Result->ResultCode));
 								Promise.EmplaceValue(Errors::FromEOSResult(Result->ResultCode));
 								return;
 							}
diff --git a/Engine/Plugins/Online/OnlineServicesEOSGS/Source/Private/Online/StatsEOSGS.cpp b/Engine/Plugins/Online/OnlineServicesEOSGS/Source/Private/Online/StatsEOSGS.cpp
--- a/Engine/Plugins/Online/OnlineServicesEOSGS/Source/Private/Online/StatsEOSGS.cpp
+++ b/Engine/Plugins/Online/OnlineServicesEOSGS/Source/Private/Online/StatsEOSGS.cpp
@@ -121,7 +121,7 @@ TOnlineAsyncOpHandle<FUpdateStats> FStatsEOSGS::UpdateStats(FUpdateStats::Params
 		{
 			if (Data->ResultCode != EOS_EResult::EOS_Success)
 			{
-				UE_LOG(LogTemp, Warning, TEXT("EOS_Stats_IngestStat failed with result=[%s]"), *LexToString(Data->ResultCode));
+				UE_LOG(LogOnlineServices, Warning, TEXT("EOS_Stats_IngestStat failed with result=[%s]"), *LexToString(Data->ResultCode));
 				InAsyncOp.SetError(Errors::FromEOSResult(Data->ResultCode));
 				return;
 			}
@@ -215,7 +215,7 @@ void FStatsEOSGS::ReadStatsFromEOSResult(const EOS_Stats_OnQueryStatsCompleteCal
 		else
 		{
 			// Put an empty stat in
-			UE_LOG(LogTemp, VeryVerbose, TEXT("Value not found for stat %s, adding empty value"), *StatName);
+			UE_LOG(LogOnlineServices, VeryVerbose, TEXT("Value not found for stat %s, adding empty value"), *StatName);
 		}
 		OutStats.Add(StatName, MoveTemp(StatValue));
 	}
@@ -243,7 +243,7 @@ TOnlineAsyncOpHandle<FQueryStats> FStatsEOSGS::QueryStats(FQueryStats::Params&&
 		{
 			if (Data->ResultCode != EOS_EResult::EOS_Success)
 			{
-				UE_LOG(LogTemp, Warning, TEXT("EOS_Stats_QueryStats failed with result=[%s]"), *LexToString(Data->ResultCode));
+				UE_LOG(LogOnlineServices, Warning, TEXT("EOS_Stats_QueryStats failed with result=[%s]"), *LexToString(Data->ResultCode));
 				InAsyncOp.SetError(Errors::FromEOSResult(Data->ResultCode));
 				return;
 			}
@@ -293,7 +293,7 @@ TOnlineAsyncOpHandle<FBatchQueryStats> FStatsEOSGS::BatchQueryStats(FBatchQueryS
 			{
 				if (Data->ResultCode != EOS_EResult::EOS_Success)
 				{
-					UE_LOG(LogTemp, Warning, TEXT("EOS_Stats_QueryStats failed with result=[%s]"), *LexToString(Data->ResultCode));
+					UE_LOG(LogOnlineServices, Warning, TEXT("EOS_Stats_QueryStats failed with result=[%s]"), *LexToString(Data->ResultCode));
 					BatchQueriedUsersStats.Reset();
 					InAsyncOp.SetError(Errors::FromEOSResult(Data->ResultCode));
 					return;
diff --git a/Engine/Plugins/Online/OnlineServicesEOSGS/Source/Private/Online/TitleFileEOSGS.cpp b/Engine/Plugins/Online/OnlineServicesEOSGS/Source/Private/Online/TitleFileEOSGS.cpp
--- a/Engine/Plugins/Online/OnlineServicesEOSGS/Source/Private/Online/TitleFileEOSGS.cpp
+++ b/Engine/Plugins/Online/OnlineServicesEOSGS/Source/Private/Online/TitleFileEOSGS.cpp
@@ -70,7 +70,7 @@ TOnlineAsyncOpHandle<FTitleFileEnumerateFiles> FTitleFileEOSGS::EnumerateFiles(F
 		{
 			if (Data->ResultCode != EOS_EResult::EOS_Success)
 			{
-				UE_LOG(LogTemp, Warning, TEXT("EOS_TitleStorage_QueryFileList failed with result=[%s]"), *LexToString(Data->ResultCode));
+				UE_LOG(LogOnlineServices, Warning, TEXT("EOS_TitleStorage_QueryFileList failed with result=[%s]"), *LexToString(Data->ResultCode));
 				InAsyncOp.SetError(Errors::FromEOSResult(Data->ResultCode));
 				return;
 			}
@@ -101,7 +101,7 @@ TOnlineAsyncOpHandle<FTitleFileEnumerateFiles> FTitleFileEOSGS::EnumerateFiles(F
 				const EOS_EResult Result = EOS_TitleStorage_CopyFileMetadataAtIndex(TitleStorageHandle, &CopyOptions, &FileMetadata);
 				if (Result != EOS_EResult::EOS_Success)
 				{
-					UE_LOG(LogTemp, Warning, TEXT("EOS_TitleStorage_CopyFileMetadataAtIndex failed with result=[%s]"), *LexToString(Result));
+					UE_LOG(LogOnlineServices, Warning, TEXT("EOS_TitleStorage_CopyFileMetadataAtIndex failed with result=[%s]"), *LexToString(Result));
 					InAsyncOp.SetError(Errors::FromEOSResult(Result));
 					return;
 				}
@@ -210,7 +210,7 @@ TOnlineAsyncOpHandle<FTitleFileReadFile> FTitleFileEOSGS::ReadFile(FTitleFileRea
 			}
 			else
 			{
-				UE_LOG(LogTemp, Warning, TEXT("EOS_TitleStorage_ReadFile failed with result=[%s]"), *LexToString(Data->ResultCode));
+				UE_LOG(LogOnlineServices, Warning, TEXT("EOS_TitleStorage_ReadFile failed with result=[%s]"), *LexToString(Data->ResultCode));
 				InAsyncOp.SetError(Errors::FromEOSResult(Data->ResultCode));
 			}
 
@@ -236,19 +236,19 @@ EOS_TitleStorage_EReadResult EOS_CALL FTitleFileEOSGS::OnReadFileDataStatic(cons
 		FileContents.Reserve(Data->TotalFileSizeBytes);
 		FileContents.Append((uint8*)Data->DataChunk, Data->DataChunkLengthBytes);
 
-		UE_LOG(LogTemp, VeryVerbose, TEXT("FTitleFileEOSGS::ReadFile ReadProgress Filename=[%s] %d/%d"), UTF8_TO_TCHAR(Data->Filename), FileContents.Num(), Data->TotalFileSizeBytes);
+		UE_LOG(LogOnlineServices, VeryVerbose, TEXT("FTitleFileEOSGS::ReadFile ReadProgress Filename=[%s] %d/%d"), UTF8_TO_TCHAR(Data->Filename), FileContents.Num(), Data->TotalFileSizeBytes);
 	}
 
 	return EOS_TitleStorage_EReadResult::EOS_TS_RR_ContinueReading;
 }
 void EOS_CALL FTitleFileEOSGS::OnFileTransferProgressStatic(const EOS_TitleStorage_FileTransferProgressCallbackInfo* Data)
 {
-	UE_LOG(LogTemp, VeryVerbose, TEXT("FTitleFileEOSGS::ReadFile TransferProgress Filename=[%s] %d/%d"), UTF8_TO_TCHAR(Data->Filename), Data->BytesTransferred, Data->TotalFileSizeBytes);
+	UE_LOG(LogOnlineServices, VeryVerbose, TEXT("FTitleFileEOSGS::ReadFile TransferProgress Filename=[%s] %d/%d"), UTF8_TO_TCHAR(Data->Filename), Data->BytesTransferred, Data->TotalFileSizeBytes);
 }
 
 void EOS_CALL FTitleFileEOSGS::OnReadFileCompleteStatic(const EOS_TitleStorage_ReadFileCallbackInfo* Data)
 {
-	UE_LOG(LogTemp, VeryVerbose, TEXT("FTitleFileEOSGS::ReadFile Complete Filename=[%s]"), UTF8_TO_TCHAR(Data->Filename));
+	UE_LOG(LogOnlineServices, VeryVerbose, TEXT("FTitleFileEOSGS::ReadFile Complete Filename=[%s]"), UTF8_TO_TCHAR(Data->Filename));
 
 	FTitleFileReadFileClientData* ClientData = (FTitleFileReadFileClientData*)Data->ClientData;
 	check(ClientData);
diff --git a/Engine/Plugins/Online/OnlineServicesEOSGS/Source/Private/Online/UserFileEOSGS.cpp b/Engine/Plugins/Online/OnlineServicesEOSGS/Source/Private/Online/UserFileEOSGS.cpp
--- a/Engine/Plugins/Online/OnlineServicesEOSGS/Source/Private/Online/UserFileEOSGS.cpp
+++ b/Engine/Plugins/Online/OnlineServicesEOSGS/Source/Private/Online/UserFileEOSGS.cpp
@@ -66,7 +66,7 @@ TOnlineAsyncOpHandle<FUserFileEnumerateFiles> FUserFileEOSGS::EnumerateFiles(FUs
 			if (Data->ResultCode != EOS_EResult::EOS_Success
 				&& Data->ResultCode != EOS_EResult::EOS_NotFound)
 			{
-				UE_LOG(LogTemp, Warning, TEXT("EOS_PlayerDataStorage_QueryFileList failed with result=[%s]"), *LexToString(Data->ResultCode));
+				UE_LOG(LogOnlineServices, Warning, TEXT("EOS_PlayerDataStorage_QueryFileList failed with result=[%s]"), *LexToString(Data->ResultCode));
 				Op.SetError(Errors::FromEOSResult(Data->ResultCode));
 				return;
 			}
@@ -82,7 +82,7 @@ TOnlineAsyncOpHandle<FUserFileEnumerateFiles> FUserFileEOSGS::EnumerateFiles(FUs
 			EOS_EResult Result = EOS_PlayerDataStorage_GetFileMetadataCount(PlayerDataStorageHandle, &GetCountOptions, &NumFiles);
 			if (Result != EOS_EResult::EOS_Success)
 			{
-				UE_LOG(LogTemp, Warning, TEXT("EOS_PlayerDataStorage_GetFileMetadataCount failed with result=[%s]"), *LexToString(Result));
+				UE_LOG(LogOnlineServices, Warning, TEXT("EOS_PlayerDataStorage_GetFileMetadataCount failed with result=[%s]"), *LexToString(Result));
 				Op.SetError(Errors::FromEOSResult(Result));
 				return;
 			}
@@ -104,7 +104,7 @@ TOnlineAsyncOpHandle<FUserFileEnumerateFiles> FUserFileEOSGS::EnumerateFiles(FUs
 				Result = EOS_PlayerDataStorage_CopyFileMetadataAtIndex(PlayerDataStorageHandle, &CopyOptions, &FileMetadata);
 				if (Result != EOS_EResult::EOS_Success)
 				{
-					UE_LOG(LogTemp, Warning, TEXT("EOS_PlayerDataStorage_CopyFileMetadataAtIndex failed with result=[%s]"), *LexToString(Result));
+					UE_LOG(LogOnlineServices, Warning, TEXT("EOS_PlayerDataStorage_CopyFileMetadataAtIndex failed with result=[%s]"), *LexToString(Result));
 					Op.SetError(Errors::FromEOSResult(Result));
 					return;
 				}
@@ -213,7 +213,7 @@ TOnlineAsyncOpHandle<FUserFileReadFile> FUserFileEOSGS::ReadFile(FUserFileReadFi
 			}
 			else
 			{
-				UE_LOG(LogTemp, Warning, TEXT("EOS_PlayerDataStorage_ReadFile failed with result=[%s]"), *LexToString(Data->ResultCode));
+				UE_LOG(LogOnlineServices, Warning, TEXT("EOS_PlayerDataStorage_ReadFile failed with result=[%s]"), *LexToString(Data->ResultCode));
 				Op.SetError(Errors::FromEOSResult(Data->ResultCode));
 			}
 
@@ -287,7 +287,7 @@ TOnlineAsyncOpHandle<FUserFileWriteFile> FUserFileEOSGS::WriteFile(FUserFileWrit
 		}
 		else
 		{
-			UE_LOG(LogTemp, Warning, TEXT("EOS_PlayerDataStorage_WriteFile failed with result=[%s]"), *LexToString(Data->ResultCode));
+			UE_LOG(LogOnlineServices, Warning, TEXT("EOS_PlayerDataStorage_WriteFile failed with result=[%s]"), *LexToString(Data->ResultCode));
 			Op.SetError(Errors::FromEOSResult(Data->ResultCode));
 		}
 
@@ -342,7 +342,7 @@ TOnlineAsyncOpHandle<FUserFileCopyFile> FUserFileEOSGS::CopyFile(FUserFileCopyFi
 		}
 		else
 		{
-			UE_LOG(LogTemp, Warning, TEXT("EOS_PlayerDataStorage_DuplicateFile failed with result=[%s]"), *LexToString(Data->ResultCode));
+			UE_LOG(LogOnlineServices, Warning, TEXT("EOS_PlayerDataStorage_DuplicateFile failed with result=[%s]"), *LexToString(Data->ResultCode));
 			Op.SetError(Errors::FromEOSResult(Data->ResultCode));
 		}
 	})
@@ -390,7 +390,7 @@ TOnlineAsyncOpHandle<FUserFileDeleteFile> FUserFileEOSGS::DeleteFile(FUserFileDe
 		}
 		else
 		{
-			UE_LOG(LogTemp, Warning, TEXT("EOS_PlayerDataStorage_DeleteFile failed with result=[%s]"), *LexToString(Data->ResultCode));
+			UE_LOG(LogOnlineServices, Warning, TEXT("EOS_PlayerDataStorage_DeleteFile failed with result=[%s]"), *LexToString(Data->ResultCode));
 			Op.SetError(Errors::FromEOSResult(Data->ResultCode));
 		}
 	})
@@ -410,7 +410,7 @@ EOS_PlayerDataStorage_EReadResult EOS_CALL FUserFileEOSGS::OnReadFileDataStatic(
 		FileContents.Reserve(Data->TotalFileSizeBytes);
 		FileContents.Append((uint8*)Data->DataChunk, Data->DataChunkLengthBytes);
 
-		UE_LOG(LogTemp, VeryVerbose, TEXT("EOS_PlayerDataStorage_ReadFile Progress Filename=[%s] %d/%d"), UTF8_TO_TCHAR(Data->Filename), FileContents.Num(), Data->TotalFileSizeBytes);
+		UE_LOG(LogOnlineServices, VeryVerbose, TEXT("EOS_PlayerDataStorage_ReadFile Progress Filename=[%s] %d/%d"), UTF8_TO_TCHAR(Data->Filename), FileContents.Num(), Data->TotalFileSizeBytes);
 	}
 
 	return EOS_PlayerDataStorage_EReadResult::EOS_RR_ContinueReading;
@@ -418,7 +418,7 @@ EOS_PlayerDataStorage_EReadResult EOS_CALL FUserFileEOSGS::OnReadFileDataStatic(
 
 void EOS_CALL FUserFileEOSGS::OnReadFileCompleteStatic(const EOS_PlayerDataStorage_ReadFileCallbackInfo* Data)
 {
-	UE_LOG(LogTemp, VeryVerbose, TEXT("EOS_PlayerDataStorage_ReadFile Complete Filename=[%s]"), UTF8_TO_TCHAR(Data->Filename));
+	UE_LOG(LogOnlineServices, VeryVerbose, TEXT("EOS_PlayerDataStorage_ReadFile Complete Filename=[%s]"), UTF8_TO_TCHAR(Data->Filename));
 
 	FUserFileReadFileClientData* ClientData = (FUserFileReadFileClientData*)Data->ClientData;
 	check(ClientData);
@@ -444,7 +444,7 @@ EOS_PlayerDataStorage_EWriteResult EOS_CALL FUserFileEOSGS::OnWriteFileDataStati
 	const uint32_t BytesWrittenTotal = BytesWrittenBefore + BytesToWrite;
 	ClientData->BytesWritten = BytesWrittenTotal;
 
-	UE_LOG(LogTemp, VeryVerbose, TEXT("EOS_PlayerDataStorage_WriteFile Progress Filename=[%s] %d/%d"), UTF8_TO_TCHAR(Data->Filename), BytesWrittenTotal, FileContents.Num());
+	UE_LOG(LogOnlineServices, VeryVerbose, TEXT("EOS_PlayerDataStorage_WriteFile Progress Filename=[%s] %d/%d"), UTF8_TO_TCHAR(Data->Filename), BytesWrittenTotal, FileContents.Num());
 
 	const bool bDone = BytesWrittenTotal == FileContents.Num();
 	if (bDone)
@@ -459,7 +459,7 @@ EOS_PlayerDataStorage_EWriteResult EOS_CALL FUserFileEOSGS::OnWriteFileDataStati
 
 void EOS_CALL FUserFileEOSGS::OnWriteFileCompleteStatic(const EOS_PlayerDataStorage_WriteFileCallbackInfo* Data)
 {
-	UE_LOG(LogTemp, VeryVerbose, TEXT("EOS_PlayerDataStorage_WriteFile Complete Filename=[%s]"), UTF8_TO_TCHAR(Data->Filename));
+	UE_LOG(LogOnlineServices, VeryVerbose, TEXT("EOS_PlayerDataStorage_WriteFile Complete Filename=[%s]"), UTF8_TO_TCHAR(Data->Filename));
 
 	FUserFileWriteFileClientData* ClientData = (FUserFileWriteFileClientData*)Data->ClientData;
 	check(ClientData);
@@ -470,7 +470,7 @@ void EOS_CALL FUserFileEOSGS::OnWriteFileCompleteStatic(const EOS_PlayerDataStor
 
 void EOS_CALL FUserFileEOSGS::OnFileTransferProgressStatic(const EOS_PlayerDataStorage_FileTransferProgressCallbackInfo* Data)
 {
-	UE_LOG(LogTemp, VeryVerbose, TEXT("FUserFileEOSGS TransferProgress Filename=[%s] %d/%d"), UTF8_TO_TCHAR(Data->Filename), Data->BytesTransferred, Data->TotalFileSizeBytes);
+	UE_LOG(LogOnlineServices, VeryVerbose, TEXT("FUserFileEOSGS TransferProgress Filename=[%s] %d/%d"), UTF8_TO_TCHAR(Data->Filename), Data->BytesTransferred, Data->TotalFileSizeBytes);
 }
 
 /* UE::Online */ }
diff --git a/Engine/Plugins/Online/OnlineServicesNull/Source/Private/Online/AchievementsNull.cpp b/Engine/Plugins/Online/OnlineServicesNull/Source/Private/Online/AchievementsNull.cpp
--- a/Engine/Plugins/Online/OnlineServicesNull/Source/Private/Online/AchievementsNull.cpp
+++ b/Engine/Plugins/Online/OnlineServicesNull/Source/Private/Online/AchievementsNull.cpp
@@ -218,9 +218,9 @@ TOnlineResult<FDisplayAchievementUI> FAchievementsNull::DisplayAchievementUI(FDi
 	}
 
 	const FAchievementState& AchievementState = LocalUserAchievementStates->FindChecked(Params.AchievementId);
-	UE_LOG(LogTemp, Display, TEXT("AchievementsNull: DisplayAchievementUI LocalAccountId=[%s]"), *ToLogString(Params.LocalAccountId));
-	UE_LOG(LogTemp, Display, TEXT("AchievementsNull: DisplayAchievementUI AchievementDefinition=[%s]"), *ToLogString(*AchievementDefinition));
-	UE_LOG(LogTemp, Display, TEXT("AchievementsNull: DisplayAchievementUI AchievementState=[%s]"), *ToLogString(AchievementState));
+	UE_LOG(LogOnlineServices, Display, TEXT("AchievementsNull: DisplayAchievementUI LocalAccountId=[%s]"), *ToLogString(Params.LocalAccountId));
+	UE_LOG(LogOnlineServices, Display, TEXT("AchievementsNull: DisplayAchievementUI AchievementDefinition=[%s]"), *ToLogString(*AchievementDefinition));
+	UE_LOG(LogOnlineServices, Display, TEXT("AchievementsNull: DisplayAchievementUI AchievementState=[%s]"), *ToLogString(AchievementState));
 
 	return TOnlineResult<FDisplayAchievementUI>(FDisplayAchievementUI::Result());
 }
diff --git a/Engine/Plugins/Online/OnlineServicesNull/Source/Private/Online/LobbiesNull.cpp b/Engine/Plugins/Online/OnlineServicesNull/Source/Private/Online/LobbiesNull.cpp
--- a/Engine/Plugins/Online/OnlineServicesNull/Source/Private/Online/LobbiesNull.cpp
+++ b/Engine/Plugins/Online/OnlineServicesNull/Source/Private/Online/LobbiesNull.cpp
@@ -107,7 +107,7 @@ TOnlineAsyncOpHandle<FCreateLobby> FLobbiesNull::CreateLobby(FCreateLobby::Param
 
 TOnlineAsyncOpHandle<FFindLobbies> FLobbiesNull::FindLobbies(FFindLobbies::Params&& Params)
 {
-	UE_LOG(LogTemp, Warning, TEXT("FLobbiesNull::FindLobbies"));
+	UE_LOG(LogOnlineServices, Warning, TEXT("FLobbiesNull::FindLobbies"));
 	if (!CurrentLobbySearch.IsValid())
 	{
 		CurrentLobbySearch = MakeShared<FFindLobbies::Result>();
@@ -151,11 +151,11 @@ TOnlineAsyncOpHandle<FJoinLobby> FLobbiesNull::JoinLobby(FJoinLobby::Params&& Pa
 	{
 		if (!Lobby.IsValid())
 		{
-			UE_LOG(LogTemp, Warning, TEXT("Failed to find a lobby with ID %d"), Op->GetParams().LobbyId.GetHandle())
+			UE_LOG(LogOnlineServices, Warning, TEXT("Failed to find a lobby with ID %d"), Op->GetParams().LobbyId.GetHandle())
 		}
 		else
 		{
-			UE_LOG(LogTemp, Warning, TEXT("Lobby with name %s already exists"), *LobbyName.ToString());
+			UE_LOG(LogOnlineServices, Warning, TEXT("Lobby with name %s already exists"), *LobbyName.ToString());
 		}
 		Op->SetError(Errors::RequestFailure());  
 		return Op->GetHandle();
@@ -211,7 +211,7 @@ TSharedRef<FLobbyNull> FLobbiesNull::CreateNamedLobby(const FCreateLobby::Params
 
 uint32 FLobbiesNull::FindLANSession()
 {
-	UE_LOG(LogTemp, Warning, TEXT("FLobbiesNull::FindLANSession"));
+	UE_LOG(LogOnlineServices, Warning, TEXT("FLobbiesNull::FindLANSession"));
 	uint32 Return = ONLINE_IO_PENDING;
 
 	// Recreate the unique identifier for this client
@@ -224,7 +224,7 @@ uint32 FLobbiesNull::FindLANSession()
 	LANSessionManager.CreateClientQueryPacket(Packet, LANSessionManager.LanNonce);
 	if (LANSessionManager.Search(Packet, ResponseDelegate, TimeoutDelegate) == false)
 	{
-		UE_LOG(LogTemp, Warning, TEXT("FLobbiesNull::FindLANSession: Failed, returning"));
+		UE_LOG(LogOnlineServices, Warning, TEXT("FLobbiesNull::FindLANSession: Failed, returning"));
 		Return = ONLINE_FAIL;
 
 		FinalizeLANSearch();
@@ -236,7 +236,7 @@ uint32 FLobbiesNull::FindLANSession()
 	}
 	else
 	{
-		UE_LOG(LogTemp, Warning, TEXT("FLobbiesNull::FindLANSession: searching...."));
+		UE_LOG(LogOnlineServices, Warning, TEXT("FLobbiesNull::FindLANSession: searching...."));
 	}
 
 
@@ -251,7 +251,7 @@ bool FLobbiesNull::NeedsToAdvertise()
 
 uint32 FLobbiesNull::UpdateLANStatus()
 {
-	UE_LOG(LogTemp, Warning, TEXT("FLobbiesNull::UpdateLANStatus"));
+	UE_LOG(LogOnlineServices, Warning, TEXT("FLobbiesNull::UpdateLANStatus"));
 	uint32 Result = ONLINE_SUCCESS;
 
 	if ( NeedsToAdvertise() )
@@ -259,11 +259,11 @@ uint32 FLobbiesNull::UpdateLANStatus()
 		// set up LAN session
 		if (LANSessionManager.GetBeaconState() == ELanBeaconState::NotUsingLanBeacon)
 		{
-			UE_LOG(LogTemp, Warning, TEXT("FLobbiesNull::UpdateLANStatus- Hosting.."));
+			UE_LOG(LogOnlineServices, Warning, TEXT("FLobbiesNull::UpdateLANStatus- Hosting.."));
 			FOnValidQueryPacketDelegate QueryPacketDelegate = FOnValidQueryPacketDelegate::CreateRaw(this, &FLobbiesNull::OnValidQueryPacketReceived);
 			if (!LANSessionManager.Host(QueryPacketDelegate))
 			{
-				UE_LOG(LogTemp, Warning, TEXT("FLobbiesNull::Failed to host!"));
+				UE_LOG(LogOnlineServices, Warning, TEXT("FLobbiesNull::Failed to host!"));
 				Result = ONLINE_FAIL;
 
 				LANSessionManager.StopLANSession();
@@ -287,7 +287,7 @@ void FLobbiesNull::AppendLobbyToPacket(FNboSerializeToBuffer& Packet, const TSha
 	using namespace NboSerializerNullSvc;
 
 	/** Owner of the session */
-	UE_LOG(LogTemp, Warning, TEXT("AppendLobbyToPacket: Appending IP %s"), *Lobby->HostAddr->ToString(true));
+	UE_LOG(LogOnlineServices, Warning, TEXT("AppendLobbyToPacket: Appending IP %s"), *Lobby->HostAddr->ToString(true));
 
 	SerializeToBuffer(Packet, Lobby->Data->OwnerAccountId);
 	SerializeToBuffer(Packet, Lobby->Data->Attributes);
@@ -302,7 +302,7 @@ void FLobbiesNull::AppendLobbyToPacket(FNboSerializeToBuffer& Packet, const TSha
 
 void FLobbiesNull::OnValidQueryPacketReceived(uint8* PacketData, int32 PacketLength, uint64 ClientNonce)
 {
-	UE_LOG(LogTemp, Warning, TEXT("FLobbiesNull::OnValidQueryPacketReceived"));
+	UE_LOG(LogOnlineServices, Warning, TEXT("FLobbiesNull::OnValidQueryPacketReceived"));
 	// Iterate through all registered sessions and respond for each one that can be joinable
 	for (TPair<FName, TSharedRef<FLobbyNull>>& Pair : NamedLobbies)
 	{
@@ -324,7 +324,7 @@ void FLobbiesNull::OnValidQueryPacketReceived(uint8* PacketData, int32 PacketLen
 		}
 		else
 		{
-			UE_LOG(LogTemp, Warning, TEXT("LAN broadcast packet overflow, cannot broadcast on LAN"));
+			UE_LOG(LogOnlineServices, Warning, TEXT("LAN broadcast packet overflow, cannot broadcast on LAN"));
 		}
 	}
 }
@@ -333,7 +333,7 @@ void FLobbiesNull::ReadLobbyFromPacket(FNboSerializeFromBuffer& Packet, const TS
 {
 	using namespace NboSerializerNullSvc;
 
-	UE_LOG(LogTemp, Verbose, TEXT("Reading session information from server"));
+	UE_LOG(LogOnlineServices, Verbose, TEXT("Reading session information from server"));
 
 	/** Owner of the session */
 	FAccountId OwningAccountId;
@@ -353,7 +353,7 @@ void FLobbiesNull::ReadLobbyFromPacket(FNboSerializeFromBuffer& Packet, const TS
 	Session->HostAddr->SetPort(Session->HostAddrPort);
 	Session->Data->Attributes.Add(FName(TEXT("ConnectAddress")), Session->HostAddr->ToString(true));
 	Session->Data->LobbyId = FOnlineLobbyIdRegistryNull::Get().GetNext(); //TODO: this should be coming from the packet, not a random GUID
-	UE_LOG(LogTemp, Warning, TEXT("ReadLobby: Got IP %s"), *Session->HostAddr->ToString(true));
+	UE_LOG(LogOnlineServices, Warning, TEXT("ReadLobby: Got IP %s"), *Session->HostAddr->ToString(true));
 
 	// todo: make all the members (not just the host
 	TSharedRef<FLobbyMember> NewMember = MakeShared<FLobbyMember>();
@@ -366,7 +366,7 @@ void FLobbiesNull::ReadLobbyFromPacket(FNboSerializeFromBuffer& Packet, const TS
 
 void FLobbiesNull::OnValidResponsePacketReceived(uint8* PacketData, int32 PacketLength)
 {
-	UE_LOG(LogTemp, Warning, TEXT("FLobbiesNull::OnValidResponsePacketReceived"));
+	UE_LOG(LogOnlineServices, Warning, TEXT("FLobbiesNull::OnValidResponsePacketReceived"));
 	if (CurrentLobbySearch.IsValid())
 	{
 		TSharedRef<FLobbyNull> Lobby = MakeShared<FLobbyNull>();
@@ -386,13 +386,13 @@ void FLobbiesNull::OnValidResponsePacketReceived(uint8* PacketData, int32 Packet
 	}
 	else
 	{
-		UE_LOG(LogTemp, Warning, TEXT("Failed to create new online game settings object"));
+		UE_LOG(LogOnlineServices, Warning, TEXT("Failed to create new online game settings object"));
 	}
 }
 
 uint32 FLobbiesNull::FinalizeLANSearch()
 {
-	UE_LOG(LogTemp, Warning, TEXT("FLobbiesNull::FinalizeLANSearch"));
+	UE_LOG(LogOnlineServices, Warning, TEXT("FLobbiesNull::FinalizeLANSearch"));
 	if (LANSessionManager.GetBeaconState() == ELanBeaconState::Searching)
 	{
 		LANSessionManager.StopLANSession();
@@ -403,14 +403,14 @@ uint32 FLobbiesNull::FinalizeLANSearch()
 
 void FLobbiesNull::OnLANSearchTimeout()
 {
-	UE_LOG(LogTemp, Warning, TEXT("FLobbiesNull::OnLANSearchTimeout"));
+	UE_LOG(LogOnlineServices, Warning, TEXT("FLobbiesNull::OnLANSearchTimeout"));
 	FinalizeLANSearch();
 
 	if (CurrentLobbySearch.IsValid())
 	{
 		if (CurrentLobbySearch->Lobbies.Num() > 0)
 		{
-			UE_LOG(LogTemp, Warning, TEXT("FLobbiesNull::OnLANSearchTimeout: Had results!"));
+			UE_LOG(LogOnlineServices, Warning, TEXT("FLobbiesNull::OnLANSearchTimeout: Had results!"));
 			// Allow game code to sort the servers (todo)
 			//CurrentSessionSearch->SortSearchResults();
 			FFindLobbies::Result Result = *CurrentLobbySearch; // copy it- not sure if moving something owned by a shared ptr is a good idea
@@ -422,7 +422,7 @@ void FLobbiesNull::OnLANSearchTimeout()
 		}
 	}
 	
-	UE_LOG(LogTemp, Warning, TEXT("FLobbiesNull::OnLANSearchTimeout: no results"));
+	UE_LOG(LogOnlineServices, Warning, TEXT("FLobbiesNull::OnLANSearchTimeout: no results"));
 	// Trigger the delegate as complete
 	CurrentLobbySearchHandle->SetError(Errors::Unknown());
 	CurrentLobbySearch.Reset();
diff --git a/Engine/Plugins/Online/OnlineServicesOSSAdapter/Source/Private/Online/LeaderboardsOSSAdapter.cpp b/Engine/Plugins/Online/OnlineServicesOSSAdapter/Source/Private/Online/LeaderboardsOSSAdapter.cpp
--- a/Engine/Plugins/Online/OnlineServicesOSSAdapter/Source/Private/Online/LeaderboardsOSSAdapter.cpp
+++ b/Engine/Plugins/Online/OnlineServicesOSSAdapter/Source/Private/Online/LeaderboardsOSSAdapter.cpp
@@ -47,7 +47,7 @@ TOnlineAsyncOpHandle<FReadEntriesForUsers> FLeaderboardsOSSAdapter::ReadEntriesF
 			const FUniqueNetIdPtr UniqueNetId = Auth->GetUniqueNetId(AccountId);
 			if (!UniqueNetId)
 			{
-				UE_LOG(LogTemp, Warning, TEXT("Failed to get unique net id for user %s"), *ToLogString(AccountId));
+				UE_LOG(LogOnlineServices, Warning, TEXT("Failed to get unique net id for user %s"), *ToLogString(AccountId));
 				Op.SetError(Errors::InvalidUser());
 				return MakeFulfilledPromise<void>().GetFuture();
 			}
diff --git a/Engine/Plugins/Online/OnlineServicesOSSAdapter/Source/Private/Online/SessionsOSSAdapter.cpp b/Engine/Plugins/Online/OnlineServicesOSSAdapter/Source/Private/Online/SessionsOSSAdapter.cpp
--- a/Engine/Plugins/Online/OnlineServicesOSSAdapter/Source/Private/Online/SessionsOSSAdapter.cpp
+++ b/Engine/Plugins/Online/OnlineServicesOSSAdapter/Source/Private/Online/SessionsOSSAdapter.cpp
@@ -104,7 +104,7 @@ FSchemaVariant GetV2SessionVariant(const FVariantData& InValue)
 	case EOnlineKeyValuePairDataType::Json:	// Intentional fallthrough
 	case EOnlineKeyValuePairDataType::MAX:	// Intentional fallthrough
 	default:
-		UE_LOG(LogTemp, Warning, TEXT("[GetV2SessionVariant] FVariantData type not supported by FSessionsVariant. No value was set."));
+		UE_LOG(LogOnlineServices, Warning, TEXT("[GetV2SessionVariant] FVariantData type not supported by FSessionsVariant. No value was set."));
 		break;
 	}
 
@@ -130,7 +130,7 @@ FVariantData GetV1VariantData(const FSchemaVariant& InValue)
 		Result.SetValue(InValue.GetString());
 		break;
 	default:
-		UE_LOG(LogTemp, Warning, TEXT("[GetV1VariantData] FSchemaVariant type not supported by FVariantData. No value was set."));
+		UE_LOG(LogOnlineServices, Warning, TEXT("[GetV1VariantData] FSchemaVariant type not supported by FVariantData. No value was set."));
 		break;
 	}
 
@@ -274,7 +274,7 @@ const FOnlineSession& FSessionOSSAdapter::GetV1Session() const
 void FSessionOSSAdapter::DumpState() const
 {
 	FSessionCommon::DumpState();
-	UE_LOG(LogTemp, Log, TEXT("V1Session: SessionId [%s]"), *V1Session.GetSessionIdStr());
+	UE_LOG(LogOnlineServices, Log, TEXT("V1Session: SessionId [%s]"), *V1Session.GetSessionIdStr());
 }
 
 /** FSessionsOSSAdapter */ 
@@ -567,7 +567,7 @@ TFuture<TOnlineResult<FFindSessions>> FSessionsOSSAdapter::FindSessionsImpl(cons
 		}
 		else
 		{
-			UE_LOG(LogTemp, Verbose, TEXT("[FSessionsOSSAdapter::FindSessionsImpl] UniqueNetId could not be found for SessionId %s"), *ToLogString(*Params.SessionId));
+			UE_LOG(LogOnlineServices, Verbose, TEXT("[FSessionsOSSAdapter::FindSessionsImpl] UniqueNetId could not be found for SessionId %s"), *ToLogString(*Params.SessionId));
 			Promise.EmplaceValue(Errors::InvalidParams());
 			return Future;
 		}
@@ -657,14 +657,14 @@ TOptional<FOnlineError> FSessionsOSSAdapter::CheckState(const FFindSessions::Par
 	const FUniqueNetIdRef LocalAccountId = Services.Get<FAuthOSSAdapter>()->GetUniqueNetId(Params.LocalAccountId).ToSharedRef();
 	if (!LocalAccountId->IsValid())
 	{
-		UE_LOG(LogTemp, Verbose, TEXT("[FSessionsOSSAdapter::FindSessionsImpl] UniqueId not found for LocalAccountId %s"), *ToLogString(Params.LocalAccountId));
+		UE_LOG(LogOnlineServices, Verbose, TEXT("[FSessionsOSSAdapter::FindSessionsImpl] UniqueId not found for LocalAccountId %s"), *ToLogString(Params.LocalAccountId));
 		
 		return TOptional<FOnlineError>(Errors::InvalidUser());
 	}
 
 	if (PendingV1SessionSearchesPerUser.Contains(Params.LocalAccountId))
 	{
-		UE_LOG(LogTemp, Verbose, TEXT("[FSessionsOSSAdapter::FindSessionsImpl] Session search already in progress for LocalAccountId %s"), *ToLogString(Params.LocalAccountId));
+		UE_LOG(LogOnlineServices, Verbose, TEXT("[FSessionsOSSAdapter::FindSessionsImpl] Session search already in progress for LocalAccountId %s"), *ToLogString(Params.LocalAccountId));
 		
 		return TOptional<FOnlineError>(Errors::AlreadyPending());
 	}
@@ -736,7 +736,7 @@ TFuture<TOnlineResult<FJoinSession>> FSessionsOSSAdapter::JoinSessionImpl(const
 			// If no result is found, the id might be expired, which we should notify
 			if (GetSessionIdRegistry().IsHandleExpired(Params.SessionId))
 			{
-				UE_LOG(LogTemp, Warning, TEXT("[%s] SessionId parameter [%s] is expired. Please call FindSessions to get an updated list of available sessions "), UTF8_TO_TCHAR(__FUNCTION__), *ToLogString(Params.SessionId));
+				UE_LOG(LogOnlineServices, Warning, TEXT("[%s] SessionId parameter [%s] is expired. Please call FindSessions to get an updated list of available sessions "), UTF8_TO_TCHAR(__FUNCTION__), *ToLogString(Params.SessionId));
 			}
 
 			Promise.EmplaceValue(MoveTemp(GetSessionByIdResult.GetErrorValue()));
@@ -794,7 +794,7 @@ TOptional<FOnlineError> FSessionsOSSAdapter::CheckState(const FJoinSession::Para
 	const FUniqueNetIdRef LocalAccountId = Services.Get<FAuthOSSAdapter>()->GetUniqueNetId(Params.LocalAccountId).ToSharedRef();
 	if (!LocalAccountId->IsValid())
 	{
-		UE_LOG(LogTemp, Verbose, TEXT("[%s] UniqueId not found for LocalAccountId %s"), UTF8_TO_TCHAR(__FUNCTION__), *ToLogString(Params.LocalAccountId));
+		UE_LOG(LogOnlineServices, Verbose, TEXT("[%s] UniqueId not found for LocalAccountId %s"), UTF8_TO_TCHAR(__FUNCTION__), *ToLogString(Params.LocalAccountId));
 		
 		return TOptional<FOnlineError>(Errors::InvalidUser());
 	}
@@ -902,7 +902,7 @@ TOptional<FOnlineError> FSessionsOSSAdapter::CheckState(const FSendSessionInvite
 	const FUniqueNetIdRef LocalAccountId = Services.Get<FAuthOSSAdapter>()->GetUniqueNetId(Params.LocalAccountId).ToSharedRef();
 	if (!LocalAccountId->IsValid())
 	{
-		UE_LOG(LogTemp, Verbose, TEXT("[%s] UniqueId not found for LocalAccountId %s"), UTF8_TO_TCHAR(__FUNCTION__), *ToLogString(Params.LocalAccountId));
+		UE_LOG(LogOnlineServices, Verbose, TEXT("[%s] UniqueId not found for LocalAccountId %s"), UTF8_TO_TCHAR(__FUNCTION__), *ToLogString(Params.LocalAccountId));
 
 		return TOptional<FOnlineError>(Errors::InvalidUser());
 	}
@@ -939,7 +939,7 @@ TOnlineResult<FGetResolvedConnectString> FSessionsOSSAdapter::GetResolvedConnect
 	else
 	{
 		// No valid session id set
-		UE_LOG(LogTemp, Verbose, TEXT("[FSessionsOSSAdapter::GetResolvedConnectString] Invalid SessionId %s"), *ToLogString(Params.SessionId));
+		UE_LOG(LogOnlineServices, Verbose, TEXT("[FSessionsOSSAdapter::GetResolvedConnectString] Invalid SessionId %s"), *ToLogString(Params.SessionId));
 		return TOnlineResult<FGetResolvedConnectString>(Errors::InvalidParams());
 	}
 }
diff --git a/Engine/Plugins/Online/OnlineServicesOSSAdapter/Source/Private/Online/StatsOSSAdapter.cpp b/Engine/Plugins/Online/OnlineServicesOSSAdapter/Source/Private/Online/StatsOSSAdapter.cpp
--- a/Engine/Plugins/Online/OnlineServicesOSSAdapter/Source/Private/Online/StatsOSSAdapter.cpp
+++ b/Engine/Plugins/Online/OnlineServicesOSSAdapter/Source/Private/Online/StatsOSSAdapter.cpp
@@ -62,7 +62,7 @@ void FStatsOSSAdapter::ConvertStatValueV2ToStatUpdateV1(const FString& StatName,
 	}
 	else
 	{
-		UE_LOG(LogTemp, Warning, TEXT("Can't find stat definition %s"), *StatName);
+		UE_LOG(LogOnlineServices, Warning, TEXT("Can't find stat definition %s"), *StatName);
 	}
 
 	OutOnlineStatUpdate.Set(OnlineStatValue, OnlineStatModificationType);
