Update OnlineServicesEOS to support named EOSShared config.

diff --git a/Engine/Plugins/Online/EOSShared/Source/EOSShared/Private/EOSSDKManager.cpp b/Engine/Plugins/Online/EOSShared/Source/EOSShared/Private/EOSSDKManager.cpp
--- a/Engine/Plugins/Online/EOSShared/Source/EOSShared/Private/EOSSDKManager.cpp
+++ b/Engine/Plugins/Online/EOSShared/Source/EOSShared/Private/EOSSDKManager.cpp
@@ -389,6 +389,16 @@ IEOSPlatformHandlePtr FEOSSDKManager::CreatePlatform(const FString& PlatformConf
 		return IEOSPlatformHandlePtr();
 	}
 
+	if (PlatformConfig->ProductId.IsEmpty() ||
+		PlatformConfig->SandboxId.IsEmpty() ||
+		PlatformConfig->DeploymentId.IsEmpty() ||
+		PlatformConfig->ClientId.IsEmpty() ||
+		PlatformConfig->ClientSecret.IsEmpty())
+	{
+		UE_LOG(LogEOSSDK, Warning, TEXT("Platform config missing required options"));
+		return IEOSPlatformHandlePtr();
+	}
+
 	TMap<FName, IEOSPlatformHandleWeakPtr>* PlatformMap = PlatformHandles.Find(PlatformConfigName);
 	if (PlatformMap)
 	{
diff --git a/Engine/Plugins/Online/OnlineServicesEOSGS/Source/Private/Online/OnlineServicesEOSGS.cpp b/Engine/Plugins/Online/OnlineServicesEOSGS/Source/Private/Online/OnlineServicesEOSGS.cpp
--- a/Engine/Plugins/Online/OnlineServicesEOSGS/Source/Private/Online/OnlineServicesEOSGS.cpp
+++ b/Engine/Plugins/Online/OnlineServicesEOSGS/Source/Private/Online/OnlineServicesEOSGS.cpp
@@ -37,7 +37,7 @@ void FOnlineServicesEOSGS::Init()
 	}
 	else
 	{
-		EOSPlatformHandle = PlatformFactory.CreatePlatform();
+		EOSPlatformHandle = PlatformFactory.CreatePlatform(InstanceName);
 	}
 	if (EOSPlatformHandle)
 	{
diff --git a/Engine/Plugins/Online/OnlineServicesEOSGS/Source/Private/Online/OnlineServicesEOSGSPlatformFactory.cpp b/Engine/Plugins/Online/OnlineServicesEOSGS/Source/Private/Online/OnlineServicesEOSGSPlatformFactory.cpp
--- a/Engine/Plugins/Online/OnlineServicesEOSGS/Source/Private/Online/OnlineServicesEOSGSPlatformFactory.cpp
+++ b/Engine/Plugins/Online/OnlineServicesEOSGS/Source/Private/Online/OnlineServicesEOSGSPlatformFactory.cpp
@@ -16,37 +16,70 @@
 #include "Online/OnlineServicesLog.h"
 #include <eos_base.h>
 
-struct FEOSPlatformConfig
+FString LoadEOSPlatformConfig(IEOSSDKManager* const SDKManager)
 {
-	FString ProductId;
-	FString SandboxId;
-	FString DeploymentId;
-	FString ClientId;
-	FString ClientSecret;
-	FString EncryptionKey;
-};
-
-FEOSPlatformConfig LoadEOSPlatformConfig()
-{
-	FEOSPlatformConfig PlatformConfig;
-	FString ConfigSection = FString::Printf(TEXT("OnlineServices.%s"), UE::Online::FOnlineServicesEOSGS::GetConfigNameStatic());
-	GConfig->GetString(*ConfigSection, TEXT("ProductId"), PlatformConfig.ProductId, GEngineIni);
-	GConfig->GetString(*ConfigSection, TEXT("SandboxId"), PlatformConfig.SandboxId, GEngineIni);
-	GConfig->GetString(*ConfigSection, TEXT("DeploymentId"), PlatformConfig.DeploymentId, GEngineIni);
-	GConfig->GetString(*ConfigSection, TEXT("ClientId"), PlatformConfig.ClientId, GEngineIni);
-	GConfig->GetString(*ConfigSection, TEXT("ClientSecret"), PlatformConfig.ClientSecret, GEngineIni);
+	FString ConfigSectionName = FString::Printf(TEXT("OnlineServices.%s"), UE::Online::FOnlineServicesEOSGS::GetConfigNameStatic());
+	if (!GConfig->DoesSectionExist(*ConfigSectionName, GEngineIni))
+	{
+		return FString();
+	}
+
+	// Check for explicit shared name to use.
+	FString PlatformConfigName;
+	GConfig->GetString(*ConfigSectionName, TEXT("PlatformConfigName"), PlatformConfigName, GEngineIni);
+	if (!PlatformConfigName.IsEmpty())
+	{
+		return PlatformConfigName;
+	}
+
+	// Check for cached config.
+	if (SDKManager->GetPlatformConfig(ConfigSectionName))
+	{
+		return ConfigSectionName;
+	}
+
+	// Check for legacy config. This should be handled by EOSShared instead, see FEOSSDKManager::GetPlatformConfig for details.
+	FEOSSDKPlatformConfig PlatformConfig;
+	PlatformConfig.Name = ConfigSectionName;
+
+	GConfig->GetString(*ConfigSectionName, TEXT("ProductId"), PlatformConfig.ProductId, GEngineIni);
+	if (PlatformConfig.ProductId.IsEmpty())
+	{
+		// If we're missing ProductId, assume we're missing the rest and instead rely on default EOSShared config.
+		return FString();
+	}
+
+	// Instead of specifying this config under the OnlineServices.EOS section, options should be moved to a new EOSSDK.Platform.<name> section so
+	// all modules relying on the EOSSDK can share the same config and platform instance. See FEOSSDKManager::GetPlatformConfig for details.
+	UE_LOG(LogOnlineServices, Warning, TEXT("[LoadEOSPlatformConfig] Using legacy config from %s, use EOSShared named config instead."), *ConfigSectionName);
+
+	GConfig->GetString(*ConfigSectionName, TEXT("SandboxId"), PlatformConfig.SandboxId, GEngineIni);
+	GConfig->GetString(*ConfigSectionName, TEXT("DeploymentId"), PlatformConfig.DeploymentId, GEngineIni);
+	GConfig->GetString(*ConfigSectionName, TEXT("ClientId"), PlatformConfig.ClientId, GEngineIni);
+	GConfig->GetString(*ConfigSectionName, TEXT("ClientSecret"), PlatformConfig.ClientSecret, GEngineIni);
 	// Config key renamed to ClientEncryptionKey as EncryptionKey gets removed from packaged builds due to IniKeyDenylist=EncryptionKey entry in BaseGame.ini.
-	GConfig->GetString(*ConfigSection, TEXT("ClientEncryptionKey"), PlatformConfig.EncryptionKey, GEngineIni);
-	return PlatformConfig;
-}
+	GConfig->GetString(*ConfigSectionName, TEXT("ClientEncryptionKey"), PlatformConfig.EncryptionKey, GEngineIni);
 
-bool IsEOSPlatformConfigValid(const FEOSPlatformConfig& InConfig)
-{
-	return !InConfig.ProductId.IsEmpty() &&
-		!InConfig.SandboxId.IsEmpty() &&
-		!InConfig.DeploymentId.IsEmpty() &&
-		!InConfig.ClientId.IsEmpty() &&
-		!InConfig.ClientSecret.IsEmpty();
+	PlatformConfig.CacheDirectory = SDKManager->GetCacheDirBase() / TEXT("OnlineServicesEOS");
+
+	PlatformConfig.bIsServer = IsRunningDedicatedServer() ? EOS_TRUE : EOS_FALSE;
+	if (!IsRunningGame())
+	{
+		PlatformConfig.bLoadingInEditor = true;
+	}
+	else
+	{
+		PlatformConfig.bWindowsEnableOverlayD3D9 = true;
+		PlatformConfig.bWindowsEnableOverlayD3D10 = true;
+		PlatformConfig.bWindowsEnableOverlayOpenGL = true;
+	}
+
+	if (!SDKManager->AddPlatformConfig(PlatformConfig))
+	{
+		return FString();
+	}
+
+	return ConfigSectionName;
 }
 
 namespace UE::Online {
@@ -70,78 +103,45 @@ void FOnlineServicesEOSGSPlatformFactory::TearDown()
 	return TLazySingleton<FOnlineServicesEOSGSPlatformFactory>::TearDown();
 }
 
-IEOSPlatformHandlePtr FOnlineServicesEOSGSPlatformFactory::CreatePlatform()
+IEOSPlatformHandlePtr FOnlineServicesEOSGSPlatformFactory::CreatePlatform(FName InstanceName)
 {
 	const FName EOSSharedModuleName = TEXT("EOSShared");
 	if (!FModuleManager::Get().IsModuleLoaded(EOSSharedModuleName))
 	{
 		FModuleManager::Get().LoadModuleChecked(EOSSharedModuleName);
 	}
+
 	IEOSSDKManager* const SDKManager = IEOSSDKManager::Get();
 	if (!SDKManager)
 	{
-		UE_LOG(LogOnlineServices, Error, TEXT("[FOnlineServicesEOSGS::Initialize] EOSSDK has not been loaded."));
+		UE_LOG(LogOnlineServices, Error, TEXT("[FOnlineServicesEOSGSPlatformFactory::CreatePlatform] EOSSDK has not been loaded."));
 		return {};
 	}
 
 	if (!SDKManager->IsInitialized())
 	{
-		UE_LOG(LogOnlineServices, Error, TEXT("[FOnlineServicesEOSGS::Initialize] EOSSDK has not been initialized."));
-		return {};
-	}
-
-	// Load config
-	FEOSPlatformConfig EOSPlatformConfig = LoadEOSPlatformConfig();
-	if (!IsEOSPlatformConfigValid(EOSPlatformConfig))
-	{
+		UE_LOG(LogOnlineServices, Error, TEXT("[FOnlineServicesEOSGSPlatformFactory::CreatePlatform] EOSSDK has not been initialized."));
 		return {};
 	}
-	const FTCHARToUTF8 ProductId(*EOSPlatformConfig.ProductId);
-	const FTCHARToUTF8 SandboxId(*EOSPlatformConfig.SandboxId);
-	const FTCHARToUTF8 DeploymentId(*EOSPlatformConfig.DeploymentId);
-	const FTCHARToUTF8 ClientId(*EOSPlatformConfig.ClientId);
-	const FTCHARToUTF8 ClientSecret(*EOSPlatformConfig.ClientSecret);
-	const FTCHARToUTF8 EncryptionKey(EOSPlatformConfig.EncryptionKey.IsEmpty() ? nullptr : *EOSPlatformConfig.EncryptionKey);
-	const FTCHARToUTF8 CacheDirectory(*(SDKManager->GetCacheDirBase() / TEXT("OnlineServicesEOS")));
-
-	EOS_Platform_Options PlatformOptions = {};
-	PlatformOptions.ApiVersion = 12;
-	UE_EOS_CHECK_API_MISMATCH(EOS_PLATFORM_OPTIONS_API_LATEST, 12);
-	PlatformOptions.Reserved = nullptr;
-	PlatformOptions.bIsServer = IsRunningDedicatedServer() ? EOS_TRUE : EOS_FALSE;
-	PlatformOptions.OverrideCountryCode = nullptr;
-	PlatformOptions.OverrideLocaleCode = nullptr;
-	// Can't check GIsEditor here because it is too soon!
-	if (!IsRunningGame())
-	{
-		PlatformOptions.Flags = EOS_PF_LOADING_IN_EDITOR;
-	}
-	else
-	{
-		PlatformOptions.Flags = EOS_PF_WINDOWS_ENABLE_OVERLAY_D3D9 | EOS_PF_WINDOWS_ENABLE_OVERLAY_D3D10 | EOS_PF_WINDOWS_ENABLE_OVERLAY_OPENGL; // Enable overlay support for D3D9/10 and OpenGL. This sample uses D3D11 or SDL.
-	}
 
-	PlatformOptions.ProductId = ProductId.Get();
-	PlatformOptions.SandboxId = SandboxId.Get();
-	PlatformOptions.DeploymentId = DeploymentId.Get();
-	PlatformOptions.ClientCredentials.ClientId = ClientId.Get();
-	PlatformOptions.ClientCredentials.ClientSecret = ClientSecret.Get();
-	PlatformOptions.EncryptionKey = EncryptionKey.Get();
-
-	if (FPlatformMisc::IsCacheStorageAvailable())
+	FString PlatformConfigName = LoadEOSPlatformConfig(SDKManager);
+	if (PlatformConfigName.IsEmpty())
 	{
-		PlatformOptions.CacheDirectory = CacheDirectory.Get();
-	}
-	else
-	{
-		PlatformOptions.CacheDirectory = nullptr;
+		// Check for default platform config that other modules may have setup.
+		PlatformConfigName = SDKManager->GetDefaultPlatformConfigName();
+		if (PlatformConfigName.IsEmpty())
+		{
+			UE_LOG(LogOnlineServices, Verbose, TEXT("[FOnlineServicesEOSGSPlatformFactory::CreatePlatform] Could not find platform config."));
+			return {};
+		}
 	}
 
-	IEOSPlatformHandlePtr EOSPlatformHandle = SDKManager->CreatePlatform(PlatformOptions);
+	IEOSPlatformHandlePtr EOSPlatformHandle = SDKManager->CreatePlatform(PlatformConfigName, InstanceName);
 	if (!EOSPlatformHandle)
 	{
-		UE_LOG(LogOnlineServices, Warning, TEXT("[FOnlineServicesEOSGS::Initialize] Failed to create platform."));
+		UE_LOG(LogOnlineServices, Warning, TEXT("[FOnlineServicesEOSGSPlatformFactory::CreatePlatform] Failed to create platform."));
 	}
+
 	return EOSPlatformHandle;
 }
 
@@ -149,7 +149,7 @@ IEOSPlatformHandlePtr FOnlineServicesEOSGSPlatformFactory::GetDefaultPlatform()
 {
 	if (!DefaultEOSPlatformHandle)
 	{
-		DefaultEOSPlatformHandle = CreatePlatform();
+		DefaultEOSPlatformHandle = CreatePlatform(NAME_None);
 	}
 
 	return DefaultEOSPlatformHandle;
diff --git a/Engine/Plugins/Online/OnlineServicesEOSGS/Source/Private/Online/OnlineServicesEOSGSPlatformFactory.h b/Engine/Plugins/Online/OnlineServicesEOSGS/Source/Private/Online/OnlineServicesEOSGSPlatformFactory.h
--- a/Engine/Plugins/Online/OnlineServicesEOSGS/Source/Private/Online/OnlineServicesEOSGSPlatformFactory.h
+++ b/Engine/Plugins/Online/OnlineServicesEOSGS/Source/Private/Online/OnlineServicesEOSGSPlatformFactory.h
@@ -28,7 +28,7 @@ public:
 	 * Create a new platform. Loads configuration from ini.
 	 * @return a new platform, or null on failure
 	 */
-	IEOSPlatformHandlePtr CreatePlatform();
+	IEOSPlatformHandlePtr CreatePlatform(FName InstanceName);
 	/**
 	 * Get the default platform.
 	 * The default platform is created on startup and can be used by the default instance of online services EOS.
