- Log something to tell the user if the TitleStorage / PlayerDataStorage interfaces are missing due to lack of encryption key
- Log a warning if it's missing despite us having an encryption key.

diff --git a/Engine/Plugins/Online/OnlineServicesEOSGS/Source/Private/Online/OnlineServicesEOSGS.cpp b/Engine/Plugins/Online/OnlineServicesEOSGS/Source/Private/Online/OnlineServicesEOSGS.cpp
--- a/Engine/Plugins/Online/OnlineServicesEOSGS/Source/Private/Online/OnlineServicesEOSGS.cpp
+++ b/Engine/Plugins/Online/OnlineServicesEOSGS/Source/Private/Online/OnlineServicesEOSGS.cpp
@@ -73,6 +73,27 @@ void FOnlineServicesEOSGS::Destroy()
 #endif
 }
 
+void FOnlineServicesEOSGS::WarnIfEncryptionKeyMissing(const FString& InterfaceName) const
+{
+	if (IEOSSDKManager* Manager = IEOSSDKManager::Get())
+	{
+		const FString& PlatformConfigName = GetEOSPlatformHandle()->GetConfigName();
+		if (const FEOSSDKPlatformConfig* Config = Manager->GetPlatformConfig(PlatformConfigName))
+		{
+			const FString& EncryptionKey = Config->EncryptionKey;
+			if (EncryptionKey.IsEmpty())
+			{
+				UE_LOG(LogOnlineServices, Verbose, TEXT("%s interface not available due to missing ClientEncryptionKey in config."), *InterfaceName);
+			}
+			else
+			{
+				// If we have an encryption key and still can't get the interface, something weird is going on.
+				UE_LOG(LogOnlineServices, Warning, TEXT("%s interface not available despite encryption key being present."), *InterfaceName);
+			}
+		}
+	}
+}
+
 void FOnlineServicesEOSGS::RegisterComponents()
 {
 	Components.Register<FAchievementsEOSGS>(*this);
@@ -85,10 +106,18 @@ void FOnlineServicesEOSGS::RegisterComponents()
 	{
 		Components.Register<FTitleFileEOSGS>(*this);
 	}
+	else
+	{
+		WarnIfEncryptionKeyMissing(TEXT("TitleStorage"));
+	}
 	if (EOS_Platform_GetPlayerDataStorageInterface(*GetEOSPlatformHandle()))
 	{
 		Components.Register<FUserFileEOSGS>(*this);
 	}
+	else
+	{
+		WarnIfEncryptionKeyMissing(TEXT("PlayerDataStorage"));
+	}
 	Super::RegisterComponents();
 }
 
diff --git a/Engine/Plugins/Online/OnlineServicesEOSGS/Source/Public/Online/OnlineServicesEOSGS.h b/Engine/Plugins/Online/OnlineServicesEOSGS/Source/Public/Online/OnlineServicesEOSGS.h
--- a/Engine/Plugins/Online/OnlineServicesEOSGS/Source/Public/Online/OnlineServicesEOSGS.h
+++ b/Engine/Plugins/Online/OnlineServicesEOSGS/Source/Public/Online/OnlineServicesEOSGS.h
@@ -37,6 +37,8 @@ public:
 protected:
 	IEOSPlatformHandlePtr EOSPlatformHandle;
 
+	void WarnIfEncryptionKeyMissing(const FString& InterfaceName) const;
+
 #if WITH_ENGINE
 	TSharedPtr<FSocketSubsystemEOS, ESPMode::ThreadSafe> SocketSubsystem;
 #endif
