Fixed notification handler being overwritten in AuthEOSGS.

diff --git a/Engine/Plugins/Online/OnlineServicesEOSGS/Source/Private/Online/AuthEOSGS.cpp b/Engine/Plugins/Online/OnlineServicesEOSGS/Source/Private/Online/AuthEOSGS.cpp
--- a/Engine/Plugins/Online/OnlineServicesEOSGS/Source/Private/Online/AuthEOSGS.cpp
+++ b/Engine/Plugins/Online/OnlineServicesEOSGS/Source/Private/Online/AuthEOSGS.cpp
@@ -1433,7 +1433,7 @@ void FAuthEOSGS::RegisterHandlers()
 		// Register for EAS connection status updates.
 		const int ApiVersion = 1;
 		UE_EOS_CHECK_API_MISMATCH(EOS_AUTH_ADDNOTIFYLOGINSTATUSCHANGED_API_LATEST, ApiVersion);
-		OnConnectAuthNotifyExpirationEOSEventRegistration = EOS_RegisterComponentEventHandler(
+		OnAuthLoginStatusChangedEOSEventRegistration = EOS_RegisterComponentEventHandler(
 			this,
 			AuthHandle,
 			ApiVersion,
@@ -1447,7 +1447,7 @@ void FAuthEOSGS::UnregisterHandlers()
 {
 	OnConnectLoginStatusChangedEOSEventRegistration = nullptr;
 	OnConnectAuthNotifyExpirationEOSEventRegistration = nullptr;
-	OnConnectAuthNotifyExpirationEOSEventRegistration = nullptr;
+	OnAuthLoginStatusChangedEOSEventRegistration = nullptr;
 }
 
 void FAuthEOSGS::OnConnectLoginStatusChanged(const EOS_Connect_LoginStatusChangedCallbackInfo* Data)
diff --git a/Engine/Plugins/Online/OnlineServicesEOSGS/Source/Public/Online/AuthEOSGS.h b/Engine/Plugins/Online/OnlineServicesEOSGS/Source/Public/Online/AuthEOSGS.h
--- a/Engine/Plugins/Online/OnlineServicesEOSGS/Source/Public/Online/AuthEOSGS.h
+++ b/Engine/Plugins/Online/OnlineServicesEOSGS/Source/Public/Online/AuthEOSGS.h
@@ -249,6 +249,7 @@ protected:
 	EOS_HConnect ConnectHandle = nullptr;
 	EOSEventRegistrationPtr OnConnectLoginStatusChangedEOSEventRegistration;
 	EOSEventRegistrationPtr OnConnectAuthNotifyExpirationEOSEventRegistration;
+	EOSEventRegistrationPtr OnAuthLoginStatusChangedEOSEventRegistration;
 	FAccountInfoRegistryEOS AccountInfoRegistryEOS;
 };
 
