Avoid always wrapping all non categorized EOS errors in Common::UnknownError

diff --git a/Engine/Plugins/Online/OnlineServicesEOSGS/Source/Private/Online/OnlineErrorEOSGS.cpp b/Engine/Plugins/Online/OnlineServicesEOSGS/Source/Private/Online/OnlineErrorEOSGS.cpp
--- a/Engine/Plugins/Online/OnlineServicesEOSGS/Source/Private/Online/OnlineErrorEOSGS.cpp
+++ b/Engine/Plugins/Online/OnlineServicesEOSGS/Source/Private/Online/OnlineErrorEOSGS.cpp
@@ -68,7 +68,7 @@ FOnlineError MapCommonEOSError(FOnlineError&& Error, EOS_EResult Result)
 	case EOS_EResult::EOS_VersionMismatch:		return Errors::IncompatibleVersion(MoveTemp(Error));
 	case EOS_EResult::EOS_LimitExceeded:		return Errors::TooManyRequests(MoveTemp(Error));
 	case EOS_EResult::EOS_TimedOut:				return Errors::Timeout(MoveTemp(Error));
-	default:									return Errors::Unknown(MoveTemp(Error));
+	default:									return MoveTemp(Error);
 	}
 }
 
