Bug fixes for OSSv2 wrapper:
* Presence data was not being updated beyond the first friend
* Joining lobby members were not getting lobby property data
* Declining lobby invitations was not possible.
* EOSSDK undue warnings on unset presence property data
* Callbacks on LobbyLeft still had the lobby marked as active.

diff --git a/Engine/Plugins/Online/OnlineServicesEOS/Source/Private/Online/AuthEOS.cpp b/Engine/Plugins/Online/OnlineServicesEOS/Source/Private/Online/AuthEOS.cpp
--- a/Engine/Plugins/Online/OnlineServicesEOS/Source/Private/Online/AuthEOS.cpp
+++ b/Engine/Plugins/Online/OnlineServicesEOS/Source/Private/Online/AuthEOS.cpp
@@ -431,7 +431,7 @@ TFuture<TArray<FAccountId>> FAuthEOS::ResolveAccountIds(const FAccountId& LocalA
 	Options.LocalUserId = GetProductUserIdChecked(LocalAccountId);
 	Options.AccountIdType = EOS_EExternalAccountType::EOS_EAT_EPIC;
 	Options.ExternalAccountIds = (const char**)EpicAccountIdStrPtrs.GetData();
-	Options.ExternalAccountIdCount = 1;
+	Options.ExternalAccountIdCount = EpicAccountIdStrPtrs.Num();
 
 	EOS_Async(EOS_Connect_QueryExternalAccountMappings, ConnectHandle, Options,
 	[this, WeakThis = AsWeak(), InEpicAccountIds, Promise = MoveTemp(Promise)](const EOS_Connect_QueryExternalAccountMappingsCallbackInfo* Data) mutable -> void
diff --git a/Engine/Plugins/Online/OnlineServicesEOSGS/Source/Private/Online/LobbiesEOSGS.cpp b/Engine/Plugins/Online/OnlineServicesEOSGS/Source/Private/Online/LobbiesEOSGS.cpp
--- a/Engine/Plugins/Online/OnlineServicesEOSGS/Source/Private/Online/LobbiesEOSGS.cpp
+++ b/Engine/Plugins/Online/OnlineServicesEOSGS/Source/Private/Online/LobbiesEOSGS.cpp
@@ -347,6 +347,44 @@ TOnlineAsyncOpHandle<FFindLobbies> FLobbiesEOSGS::FindLobbies(FFindLobbies::Para
 
 TOnlineAsyncOpHandle<FJoinLobby> FLobbiesEOSGS::JoinLobby(FJoinLobby::Params&& InParams)
 {
+	// Helper lambda for simplifying error handling.
+	auto LeaveLobbyDuringJoin = [this](
+		TOnlineAsyncOp<FJoinLobby>& InAsyncOp,
+		const FAccountId LocalAccountId,
+		const TSharedPtr<FLobbyDataEOS> LobbyData,
+		FOnlineError ErrorResult) -> TFuture<void>
+	{
+		// Failed to set attributes - leave the lobby.
+		FLobbiesLeaveLobbyImpl::Params LeaveLobbyParams;
+		LeaveLobbyParams.LobbyData = LobbyData;
+		LeaveLobbyParams.LocalAccountId = LocalAccountId;
+
+		TPromise<void> Promise;
+		TFuture<void> Future = Promise.GetFuture();
+
+		LeaveLobbyImpl(MoveTemp(LeaveLobbyParams))
+		.Next(
+		[
+			LocalAccountId = LocalAccountId,
+			LobbyData,
+			InAsyncOp = InAsyncOp.AsShared(),
+			ErrorResult = MoveTemp(ErrorResult),
+			Promise = MoveTemp(Promise)
+		](TDefaultErrorResult<FLobbiesLeaveLobbyImpl>&& Result) mutable
+		{
+			if (Result.IsError())
+			{
+				UE_LOG(LogTemp, Warning, TEXT("[FLobbiesEOSGS::JoinLobby] LeaveLobbyImpl Failed: User[%s], Lobby[%s], Result[%s]"),
+					*ToLogString(LocalAccountId), *LobbyData->GetLobbyIdString(), *Result.GetErrorValue().GetLogString());
+			}
+
+			InAsyncOp->SetError(MoveTemp(ErrorResult));
+			Promise.EmplaceValue();
+		});
+
+		return Future;
+	};
+
 	TOnlineAsyncOpRef<FJoinLobby> Op = GetOp<FJoinLobby>(MoveTemp(InParams));
 	const FJoinLobby::Params& Params = Op->GetParams();
 
@@ -437,8 +475,57 @@ TOnlineAsyncOpHandle<FJoinLobby> FLobbiesEOSGS::JoinLobby(FJoinLobby::Params&& I
 			InAsyncOp.SetError(Errors::FromEOSResult(Data->ResultCode));
 			return;
 		}
+	})	
+	// Step 3. Promote lobby details source to active.
+	.Then([this, LeaveLobbyDuringJoin](TOnlineAsyncOp<FJoinLobby>& InAsyncOp)
+	{
+		const FJoinLobby::Params& Params = InAsyncOp.GetParams();
+		const TSharedRef<FLobbyDataEOS>& LobbyData = GetOpDataChecked<TSharedRef<FLobbyDataEOS>>(InAsyncOp, UE_ONLINE_LOBBY_EOS_KEY_NAME_LOBBY_DATA);
+
+		// Create an active lobby details object associated with the lobby data that we can use when we get lobby update notifications.
+		const FString& LobbyIdString = LobbyData->GetLobbyIdString();
+		InAsyncOp.Data.Set<FString>(UE_ONLINE_LOBBY_EOS_KEY_NAME_LOBBY_ID_STRING, LobbyIdString);
+
+		// Try to create the lobby details object.
+		TDefaultErrorResultInternal<TSharedRef<FLobbyDetailsEOS>> Result =
+			FLobbyDetailsEOS::CreateFromLobbyId(LobbyPrerequisites.ToSharedRef(), Params.LocalAccountId, LobbyIdString);
+		if (Result.IsError())
+		{
+			UE_LOG(LogTemp, Warning, TEXT("[FLobbiesEOSGS::JoinLobby] FLobbyDetailsEOS::CreateFromLobbyId Failed: User[%s], Lobby[%s], Result[%s]"),
+				*ToLogString(Params.LocalAccountId), *LobbyIdString, *Result.GetErrorValue().GetLogString());
+			return LeaveLobbyDuringJoin(InAsyncOp, Params.LocalAccountId, LobbyData, MoveTemp(Result.GetErrorValue()));
+		}
+		else
+		{
+			// Attach the lobby details to the operation so we can register them in the next step.
+			InAsyncOp.Data.Set<TSharedRef<FLobbyDetailsEOS>>(UE_ONLINE_LOBBY_EOS_KEY_NAME_LOBBY_DETAILS, Result.GetOkValue());
+			return MakeFulfilledPromise<void>().GetFuture();
+		}
+	})	
+	// Step 4: Create the lobby data object from the lobby details.
+	.Then([this](TOnlineAsyncOp<FJoinLobby>& InAsyncOp)
+	{
+		TSharedRef<FLobbyDetailsEOS> LobbyDetails = GetOpDataChecked<TSharedRef<FLobbyDetailsEOS>>(InAsyncOp, UE_ONLINE_LOBBY_EOS_KEY_NAME_LOBBY_DETAILS);
+		return LobbyDataRegistry->FindOrCreateFromLobbyDetails(InAsyncOp.GetParams().LocalAccountId, LobbyDetails);
+	})
+	// Step 5: Handle possible errors from FindOrCreateFromLobbyDetails.
+	.Then([this, LeaveLobbyDuringJoin](TOnlineAsyncOp<FJoinLobby>& InAsyncOp, TDefaultErrorResultInternal<TSharedRef<FLobbyDataEOS>>&& Result)
+	{
+		const FJoinLobby::Params& Params = InAsyncOp.GetParams();
+		const TSharedRef<FLobbyDataEOS>& LobbyData = GetOpDataChecked<TSharedRef<FLobbyDataEOS>>(InAsyncOp, UE_ONLINE_LOBBY_EOS_KEY_NAME_LOBBY_DATA);
+
+		if (Result.IsError())
+		{
+			UE_LOG(LogTemp, Warning, TEXT("[FLobbiesEOSGS::JoinLobby] FLobbyDataRegistryEOS::FindOrCreateFromLobbyDetails Failed: User[%s], Lobby[%s], Result[%s]"),
+				*ToLogString(Params.LocalAccountId), *LobbyData->GetLobbyIdString(), *Result.GetErrorValue().GetLogString());
+			return LeaveLobbyDuringJoin(InAsyncOp, Params.LocalAccountId, LobbyData, MoveTemp(Result.GetErrorValue()));
+		}
+		else
+		{		
+			return MakeFulfilledPromise<void>().GetFuture();
+		}
 	})
-	// Step 3. Set member attributes.
+	// Step 6. Set member attributes.
 	.Then([this](TOnlineAsyncOp<FJoinLobby>& InAsyncOp)
 	{
 		const FJoinLobby::Params& Params = InAsyncOp.GetParams();
@@ -463,8 +550,8 @@ TOnlineAsyncOpHandle<FJoinLobby> FLobbiesEOSGS::JoinLobby(FJoinLobby::Params&& I
 		ModifyLobbyDataParams.ServiceChanges = MoveTemp(PrepareResult.GetOkValue().ServiceChanges);
 		return ModifyLobbyDataImpl(MoveTemp(ModifyLobbyDataParams));
 	})
-	// Step 4. Handle result.
-	.Then([this](TOnlineAsyncOp<FJoinLobby>& InAsyncOp, TDefaultErrorResult<FLobbiesModifyLobbyDataImpl>&& Result)
+	// Step 7. Handle result.
+	.Then([this, LeaveLobbyDuringJoin](TOnlineAsyncOp<FJoinLobby>& InAsyncOp, TDefaultErrorResult<FLobbiesModifyLobbyDataImpl>&& Result)
 	{
 		if (Result.IsError())
 		{
@@ -474,47 +561,25 @@ TOnlineAsyncOpHandle<FJoinLobby> FLobbiesEOSGS::JoinLobby(FJoinLobby::Params&& I
 			UE_LOG(LogOnlineServices, Warning, TEXT("[FLobbiesEOSGS::JoinLobby] ModifyLobbyMemberDataImpl Failed: User[%s], Lobby[%s], Result[%s]"),
 				*ToLogString(Params.LocalAccountId), *LobbyData->GetLobbyIdString(), *Result.GetErrorValue().GetLogString());
 
-			// Failed to set attributes - leave the lobby.
-			FLobbiesLeaveLobbyImpl::Params LeaveLobbyParams;
-			LeaveLobbyParams.LobbyData = LobbyData;
-			LeaveLobbyParams.LocalAccountId = Params.LocalAccountId;
-
-			TPromise<void> Promise;
-			TFuture<void> Future = Promise.GetFuture();
-
-			LeaveLobbyImpl(MoveTemp(LeaveLobbyParams))
-			.Next(
-			[
-				LocalAccountId = Params.LocalAccountId,
-				LobbyData,
-				InAsyncOp = InAsyncOp.AsShared(),
-				ErrorResult = MoveTemp(Result.GetErrorValue()),
-				Promise = MoveTemp(Promise)
-			](TDefaultErrorResult<FLobbiesLeaveLobbyImpl>&& Result) mutable
-			{
-				if (Result.IsError())
-				{
-					UE_LOG(LogOnlineServices, Warning, TEXT("[FLobbiesEOSGS::JoinLobby] LeaveLobbyImpl Failed: User[%s], Lobby[%s], Result[%s]"),
-						*ToLogString(LocalAccountId), *LobbyData->GetLobbyIdString(), *Result.GetErrorValue().GetLogString());
-				}
-
-				InAsyncOp->SetError(MoveTemp(ErrorResult));
-				Promise.EmplaceValue();
-			});
-
-			return Future;
+			return LeaveLobbyDuringJoin(InAsyncOp, Params.LocalAccountId, LobbyData, MoveTemp(Result.GetErrorValue()));
 		}
 		else
 		{
 			return MakeFulfilledPromise<void>().GetFuture();
 		}
 	})
-	// Step 5. Bookkeeping and notifications.
+	// Step 8. Bookkeeping and notifications.
 	.Then([this](TOnlineAsyncOp<FJoinLobby>& InAsyncOp)
 	{
 		const FJoinLobby::Params& Params = InAsyncOp.GetParams();
 		const TSharedRef<FLobbyDataEOS>& LobbyData = GetOpDataChecked<TSharedRef<FLobbyDataEOS>>(InAsyncOp, UE_ONLINE_LOBBY_EOS_KEY_NAME_LOBBY_DATA);
 
+		// Remove the invitation if this is a join in response to an invite.
+		if (TSharedPtr<FLobbyInviteDataEOS> InviteData = GetActiveInvite(Params.LocalAccountId, Params.LobbyId))
+		{
+			RemoveActiveInvite(InviteData.ToSharedRef());
+		}
+
 		// Mark the lobby active.
 		AddActiveLobby(Params.LocalAccountId, Params.bPresenceEnabled, LobbyData);
 
@@ -603,13 +668,13 @@ TOnlineAsyncOpHandle<FLeaveLobby> FLobbiesEOSGS::LeaveLobby(FLeaveLobby::Params&
 		const FLeaveLobby::Params& Params = InAsyncOp.GetParams();
 		const TSharedRef<FLobbyDataEOS>& LobbyData = GetOpDataChecked<TSharedRef<FLobbyDataEOS>>(InAsyncOp, UE_ONLINE_LOBBY_EOS_KEY_NAME_LOBBY_DATA);
 
-		// Update local cache and fire events.
-		LobbyData->GetLobbyClientData()->CommitClientChanges({ &LobbyEvents });
-
 		// Remove the lobby from the active list for the user.
 		// The lobby data will be cleaned up once all references are removed.
 		RemoveActiveLobby(Params.LocalAccountId, LobbyData);
 
+		// Update local cache and fire events.
+		LobbyData->GetLobbyClientData()->CommitClientChanges({ &LobbyEvents });
+
 		UE_LOG(LogOnlineServices, Log, TEXT("[FLobbiesEOSGS::LeaveLobby] Succeeded: User[%s], Lobby[%s]"),
 			*ToLogString(Params.LocalAccountId), *LobbyData->GetLobbyIdString());
 		InAsyncOp.SetResult(FLeaveLobby::Result{});
@@ -707,6 +772,15 @@ TOnlineAsyncOpHandle<FDeclineLobbyInvitation> FLobbiesEOSGS::DeclineLobbyInvitat
 		return Op->GetHandle();
 	}
 
+	TSharedPtr<FLobbyDataEOS> LobbyData = LobbyDataRegistry->Find(Params.LobbyId);
+	if (!LobbyData)
+	{
+		UE_LOG(LogTemp, Warning, TEXT("[FLobbiesEOSGS::JoinLobby] Failed: Unable to find lobby data. LobbyId[%s]"), *ToLogString(Params.LobbyId));
+		Op->SetError(Errors::InvalidParams());
+		return Op->GetHandle();
+	}
+	Op->Data.Set<TSharedRef<FLobbyDataEOS>>(UE_ONLINE_LOBBY_EOS_KEY_NAME_LOBBY_DATA, LobbyData.ToSharedRef());
+
 	// Start operation.
 	Op->Then([this](TOnlineAsyncOp<FDeclineLobbyInvitation>& InAsyncOp)
 	{
diff --git a/Engine/Plugins/Online/OnlineSubsystemEOS/Source/OnlineSubsystemEOS/Private/UserManagerEOS.cpp b/Engine/Plugins/Online/OnlineSubsystemEOS/Source/OnlineSubsystemEOS/Private/UserManagerEOS.cpp
--- a/Engine/Plugins/Online/OnlineSubsystemEOS/Source/OnlineSubsystemEOS/Private/UserManagerEOS.cpp
+++ b/Engine/Plugins/Online/OnlineSubsystemEOS/Source/OnlineSubsystemEOS/Private/UserManagerEOS.cpp
@@ -2656,15 +2656,19 @@ void FUserManagerEOS::SetPresence(const FUniqueNetId& UserId, const FOnlineUserP
 		Record.Key = RawString.Key.Get();
 		Record.Value = RawString.Value.Get();
 	}
-	EOS_PresenceModification_SetDataOptions DataOptions = { };
-	DataOptions.ApiVersion = 1;
-	UE_EOS_CHECK_API_MISMATCH(EOS_PRESENCE_SETDATA_API_LATEST, 1);
-	DataOptions.RecordsCount = Records.Num();
-	DataOptions.Records = Records.GetData();
-	EOS_EResult SetDataResult = EOS_PresenceModification_SetData(ChangeHandle, &DataOptions);
-	if (SetDataResult != EOS_EResult::EOS_Success)
-	{
-		UE_LOG_ONLINE(Error, TEXT("EOS_PresenceModification_SetData() failed with result code (%s)"), *LexToString(SetDataResult));
+
+	if (Records.Num() > 0)
+	{
+		EOS_PresenceModification_SetDataOptions DataOptions = { };
+		DataOptions.ApiVersion = 1;
+		UE_EOS_CHECK_API_MISMATCH(EOS_PRESENCE_SETDATA_API_LATEST, 1);
+		DataOptions.RecordsCount = Records.Num();
+		DataOptions.Records = Records.GetData();
+		EOS_EResult SetDataResult = EOS_PresenceModification_SetData(ChangeHandle, &DataOptions);
+		if (SetDataResult != EOS_EResult::EOS_Success)
+		{
+			UE_LOG_ONLINE(Error, TEXT("EOS_PresenceModification_SetData() failed with result code (%s)"), *LexToString(SetDataResult));
+		}
 	}
 
 	FSetPresenceCallback* CallbackObj = new FSetPresenceCallback(AsWeak());
