PR #11273: Fix lobby member info retrieval in OSSv2

#11273

diff --git a/Engine/Plugins/Online/OnlineServicesEOSGS/Source/Private/Online/LobbiesEOSGSTypes.cpp b/Engine/Plugins/Online/OnlineServicesEOSGS/Source/Private/Online/LobbiesEOSGSTypes.cpp
--- a/Engine/Plugins/Online/OnlineServicesEOSGS/Source/Private/Online/LobbiesEOSGSTypes.cpp
+++ b/Engine/Plugins/Online/OnlineServicesEOSGS/Source/Private/Online/LobbiesEOSGSTypes.cpp
@@ -330,7 +330,28 @@ TFuture<TDefaultErrorResultInternal<FLobbyServiceSnapshot>> FLobbyDetailsEOS::Ge
 		UE_EOS_CHECK_API_MISMATCH(EOS_LOBBYDETAILS_GETMEMBERBYINDEX_API_LATEST, 1);
 		GetMemberByIndexOptions.MemberIndex = MemberIndex;
 
-		MemberProductUserIds->Emplace(EOS_LobbyDetails_GetMemberByIndex(LobbyDetailsHandle, &GetMemberByIndexOptions));
+		// This method will return null if the querying user is not part of the lobby, in which case we can stop iterating
+		EOS_ProductUserId MemberId = EOS_LobbyDetails_GetMemberByIndex(LobbyDetailsHandle, &GetMemberByIndexOptions);
+		if (MemberId != nullptr)
+		{
+			MemberProductUserIds->Emplace(MemberId);
+		}
+		else
+		{
+			break;
+		}
+	}
+
+	EOS_LobbyDetails_GetLobbyOwnerOptions GetLobbyOwnerOptions = {};
+	GetLobbyOwnerOptions.ApiVersion = 1;
+	UE_EOS_CHECK_API_MISMATCH(EOS_LOBBYDETAILS_GETLOBBYOWNER_API_LATEST, 1);
+
+	const EOS_ProductUserId LobbyOwner = EOS_LobbyDetails_GetLobbyOwner(LobbyDetailsHandle, &GetLobbyOwnerOptions);
+
+	// In the case where the querying user is not part of the lobby and no members can be retrieved, we still want to at least resolve the id of the lobby owner, which we can retrieve at all times
+	if (MemberProductUserIds->IsEmpty())
+	{
+		MemberProductUserIds->Emplace(LobbyOwner);
 	}
 
 	TPromise<TDefaultErrorResultInternal<FLobbyServiceSnapshot>> Promise;
@@ -338,7 +359,7 @@ TFuture<TDefaultErrorResultInternal<FLobbyServiceSnapshot>> FLobbyDetailsEOS::Ge
 
 	// Resolve lobby member product user ids to FAccountId before proceeding.
 	AuthInterface->ResolveAccountIds(AssociatedLocalUser, *MemberProductUserIds)
-	.Next([StrongThis = AsShared(), Promise = MoveTemp(Promise), MemberProductUserIds](TArray<FAccountId>&& ResolvedAccountIds) mutable
+	.Next([StrongThis = AsShared(), Promise = MoveTemp(Promise), MemberProductUserIds, LobbyOwner](TArray<FAccountId>&& ResolvedAccountIds) mutable
 	{
 		if (MemberProductUserIds->Num() != ResolvedAccountIds.Num())
 		{
@@ -354,12 +375,6 @@ TFuture<TDefaultErrorResultInternal<FLobbyServiceSnapshot>> FLobbyDetailsEOS::Ge
 
 		// Resolve member info.
 		{
-			EOS_LobbyDetails_GetLobbyOwnerOptions GetLobbyOwnerOptions = {};
-			GetLobbyOwnerOptions.ApiVersion = 1;
-			UE_EOS_CHECK_API_MISMATCH(EOS_LOBBYDETAILS_GETLOBBYOWNER_API_LATEST, 1);
-
-			const EOS_ProductUserId LobbyOwner = EOS_LobbyDetails_GetLobbyOwner(StrongThis->LobbyDetailsHandle, &GetLobbyOwnerOptions);
-
 			for (int32 MemberIndex = 0; MemberIndex < MemberProductUserIds->Num(); ++MemberIndex)
 			{
 				const EOS_ProductUserId MemberProductUserId = (*MemberProductUserIds)[MemberIndex];
