Support OSSv2 Account Portal in iOS

#11274

#changelist validated
#virtualized

diff --git a/Engine/Plugins/Online/OnlineServicesEOSGS/Source/OnlineServicesEOSGS.Build.cs b/Engine/Plugins/Online/OnlineServicesEOSGS/Source/OnlineServicesEOSGS.Build.cs
--- a/Engine/Plugins/Online/OnlineServicesEOSGS/Source/OnlineServicesEOSGS.Build.cs
+++ b/Engine/Plugins/Online/OnlineServicesEOSGS/Source/OnlineServicesEOSGS.Build.cs
@@ -36,5 +36,11 @@ public class OnlineServicesEOSGS : ModuleRules
 				}
 			);
 		}
+
+		if (Target.Platform == UnrealTargetPlatform.IOS)
+		{
+			PublicWeakFrameworks.Add("AuthenticationServices");
+			PrivateDependencyModuleNames.Add("ApplicationCore");
+		}
 	}
 }
diff --git a/Engine/Plugins/Online/OnlineServicesEOSGS/Source/Private/Online/AuthEOSGS.cpp b/Engine/Plugins/Online/OnlineServicesEOSGS/Source/Private/Online/AuthEOSGS.cpp
--- a/Engine/Plugins/Online/OnlineServicesEOSGS/Source/Private/Online/AuthEOSGS.cpp
+++ b/Engine/Plugins/Online/OnlineServicesEOSGS/Source/Private/Online/AuthEOSGS.cpp
@@ -7,11 +7,21 @@
 #include "IEOSSDKManager.h"
 #include "Misc/CommandLine.h"
 #include "Online/AuthErrors.h"
+#include "Online/EOSAuthLoginOptionsCommon.h"
 #include "Online/OnlineErrorEOSGS.h"
 #include "Online/OnlineIdEOSGS.h"
 #include "Online/OnlineServicesEOSGS.h"
 #include "Online/OnlineUtils.h"
 
+#if __has_include(COMPILED_PLATFORM_HEADER(EOSAuthLoginOptions.h))
+#include COMPILED_PLATFORM_HEADER(EOSAuthLoginOptions.h)
+#else 
+namespace UE::Online
+{
+    using FPlatformEOSAuthLoginOptions = FEOSAuthLoginOptionsCommon;
+}
+#endif
+
 #include "eos_auth.h"
 #include "eos_connect.h"
 
@@ -187,34 +197,9 @@ TDefaultErrorResultInternal<FEOSConnectLoginOptions> FEOSConnectLoginOptions::Cr
 	return TDefaultErrorResultInternal<FEOSConnectLoginOptions>(MoveTemp(EOSConnectLoginOptions));
 }
 
-enum class EEOSAuthTranslationFlags : uint8
-{
-	None = 0,
-	SetId = 1 << 0,
-	SetTokenFromString = 1 << 1,
-	SetTokenFromExternalAuth = 1 << 2,
-};
-ENUM_CLASS_FLAGS(EEOSAuthTranslationFlags);
-
-struct FEOSAuthTranslationTraits
-{
-	EOS_ELoginCredentialType Type;
-	EEOSAuthTranslationFlags Flags;
-};
-
-struct FEOSExternalAuthTranslationTraits
-{
-	EOS_EExternalCredentialType Type;
-};
-
-class FEOSAuthLoginOptions : public EOS_Auth_LoginOptions
+class FEOSAuthLoginOptions : public FPlatformEOSAuthLoginOptions
 {
 public:
-	FEOSAuthLoginOptions(const FEOSAuthLoginOptions&) = delete;
-	FEOSAuthLoginOptions& operator=(const FEOSAuthLoginOptions&) = delete;
-	FEOSAuthLoginOptions(FEOSAuthLoginOptions&&);
-	FEOSAuthLoginOptions& operator=(FEOSAuthLoginOptions&&);
-
 	static TDefaultErrorResultInternal<FEOSAuthLoginOptions> Create(
 		FName CredentialsType,
 		const FString& CredentialsId,
@@ -228,109 +213,9 @@ private:
 		const TVariant<FString, FExternalAuthToken>& CredentialsToken,
 		const TArray<FString>& Scopes);
 
-	FEOSAuthLoginOptions();
-
-	static const FEOSAuthTranslationTraits* GetLoginTranslatorTraits(FName Name);
-	static const FEOSExternalAuthTranslationTraits* GetExternalAuthTranslationTraits(FName ExternalAuthType);
-
-	EOS_Auth_Credentials CredentialsData;
-	TArray<char> IdUtf8;
-	TArray<char> TokenUtf8;
+	FEOSAuthLoginOptions() = default;
 };
 
-const FEOSAuthTranslationTraits* FEOSAuthLoginOptions::GetLoginTranslatorTraits(FName Name)
-{
-	static const TMap<FName, FEOSAuthTranslationTraits> SupportedLoginTranslatorTraits = {
-		{ LoginCredentialsType::Password, { EOS_ELoginCredentialType::EOS_LCT_Password, EEOSAuthTranslationFlags::SetId | EEOSAuthTranslationFlags::SetTokenFromString } },
-		{ LoginCredentialsType::ExchangeCode, { EOS_ELoginCredentialType::EOS_LCT_ExchangeCode, EEOSAuthTranslationFlags::SetTokenFromString } },
-		{ LoginCredentialsType::PersistentAuth, { EOS_ELoginCredentialType::EOS_LCT_PersistentAuth, EEOSAuthTranslationFlags::None } },
-		{ LoginCredentialsType::Developer, { EOS_ELoginCredentialType::EOS_LCT_Developer, EEOSAuthTranslationFlags::SetId | EEOSAuthTranslationFlags::SetTokenFromString } },
-		{ LoginCredentialsType::RefreshToken, { EOS_ELoginCredentialType::EOS_LCT_RefreshToken, EEOSAuthTranslationFlags::SetTokenFromString } },
-		{ LoginCredentialsType::AccountPortal, { EOS_ELoginCredentialType::EOS_LCT_AccountPortal, EEOSAuthTranslationFlags::SetId | EEOSAuthTranslationFlags::SetTokenFromString } },
-		{ LoginCredentialsType::ExternalAuth, { EOS_ELoginCredentialType::EOS_LCT_ExternalAuth, EEOSAuthTranslationFlags::SetTokenFromExternalAuth } },
-	};
-
-	return SupportedLoginTranslatorTraits.Find(Name);
-}
-
-const FEOSExternalAuthTranslationTraits* FEOSAuthLoginOptions::GetExternalAuthTranslationTraits(FName ExternalAuthType)
-{
-	static const TMap<FName, FEOSExternalAuthTranslationTraits> SupportedExternalAuthTraits = {
-		{ ExternalLoginType::Epic, { EOS_EExternalCredentialType::EOS_ECT_EPIC } },
-		{ ExternalLoginType::SteamSessionTicket, { EOS_EExternalCredentialType::EOS_ECT_STEAM_SESSION_TICKET } },
-		{ ExternalLoginType::PsnIdToken, { EOS_EExternalCredentialType::EOS_ECT_PSN_ID_TOKEN } },
-		{ ExternalLoginType::XblXstsToken, { EOS_EExternalCredentialType::EOS_ECT_XBL_XSTS_TOKEN } },
-		{ ExternalLoginType::DiscordAccessToken, { EOS_EExternalCredentialType::EOS_ECT_DISCORD_ACCESS_TOKEN } },
-		{ ExternalLoginType::GogSessionTicket, { EOS_EExternalCredentialType::EOS_ECT_GOG_SESSION_TICKET } },
-		{ ExternalLoginType::NintendoIdToken, { EOS_EExternalCredentialType::EOS_ECT_NINTENDO_ID_TOKEN } },
-		{ ExternalLoginType::NintendoNsaIdToken, { EOS_EExternalCredentialType::EOS_ECT_NINTENDO_NSA_ID_TOKEN } },
-		{ ExternalLoginType::UplayAccessToken, { EOS_EExternalCredentialType::EOS_ECT_UPLAY_ACCESS_TOKEN } },
-		{ ExternalLoginType::OpenIdAccessToken, { EOS_EExternalCredentialType::EOS_ECT_OPENID_ACCESS_TOKEN } },
-		{ ExternalLoginType::DeviceIdAccessToken, { EOS_EExternalCredentialType::EOS_ECT_DEVICEID_ACCESS_TOKEN } },
-		{ ExternalLoginType::AppleIdToken, { EOS_EExternalCredentialType::EOS_ECT_APPLE_ID_TOKEN } },
-		{ ExternalLoginType::GoogleIdToken, { EOS_EExternalCredentialType::EOS_ECT_GOOGLE_ID_TOKEN } },
-		{ ExternalLoginType::OculusUserIdNonce, { EOS_EExternalCredentialType::EOS_ECT_OCULUS_USERID_NONCE } },
-		{ ExternalLoginType::ItchioJwt, { EOS_EExternalCredentialType::EOS_ECT_ITCHIO_JWT } },
-		{ ExternalLoginType::ItchioKey, { EOS_EExternalCredentialType::EOS_ECT_ITCHIO_KEY } },
-		{ ExternalLoginType::EpicIdToken, { EOS_EExternalCredentialType::EOS_ECT_EPIC_ID_TOKEN } },
-		{ ExternalLoginType::AmazonAccessToken, { EOS_EExternalCredentialType::EOS_ECT_AMAZON_ACCESS_TOKEN } },
-	};
-
-	return SupportedExternalAuthTraits.Find(ExternalAuthType);
-}
-
-FEOSAuthLoginOptions::FEOSAuthLoginOptions(FEOSAuthLoginOptions&& Other)
-{
-	*this = MoveTemp(Other);
-}
-
-FEOSAuthLoginOptions& FEOSAuthLoginOptions::operator=(FEOSAuthLoginOptions&& Other)
-{
-	CredentialsData = Other.CredentialsData;
-
-	// Pointer fixup.
-	if (CredentialsData.Id)
-	{
-		IdUtf8 = MoveTemp(Other.IdUtf8);
-		CredentialsData.Id = IdUtf8.GetData();
-	}
-	if (CredentialsData.Token)
-	{
-		TokenUtf8 = MoveTemp(Other.TokenUtf8);
-		CredentialsData.Token = TokenUtf8.GetData();
-	}
-	if (CredentialsData.SystemAuthCredentialsOptions)
-	{
-		// todo
-	}
-
-	Credentials = &CredentialsData;
-	ApiVersion = Other.ApiVersion;
-	ScopeFlags = Other.ScopeFlags;
-
-	Other.Credentials = nullptr;
-	return *this;
-}
-
-FEOSAuthLoginOptions::FEOSAuthLoginOptions()
-{
-	// EOS_Auth_LoginOptions init
-	UE_EOS_CHECK_API_MISMATCH(EOS_AUTH_LOGIN_API_LATEST, 3);
-	ApiVersion = 2;
-	Credentials = &CredentialsData;
-	ScopeFlags = EOS_EAuthScopeFlags::EOS_AS_NoFlags;
-	LoginFlags = 0;
-
-	// EOS_Auth_Credentials init
-	UE_EOS_CHECK_API_MISMATCH(EOS_AUTH_CREDENTIALS_API_LATEST, 4);
-	CredentialsData.ApiVersion = 4;
-	CredentialsData.Id = nullptr;
-	CredentialsData.Token = nullptr;
-	CredentialsData.Type = EOS_ELoginCredentialType::EOS_LCT_Password;
-	CredentialsData.SystemAuthCredentialsOptions = nullptr;
-	CredentialsData.ExternalType = EOS_EExternalCredentialType::EOS_ECT_EPIC;
-}
-
 TDefaultErrorResultInternal<FEOSAuthLoginOptions> FEOSAuthLoginOptions::Create(
 	FName CredentialsType,
 	const FString& CredentialsId,
@@ -425,7 +310,11 @@ TDefaultErrorResultInternal<FEOSAuthLoginOptions> FEOSAuthLoginOptions::CreateIm
 		EOSAuthLoginOptions.CredentialsData.ExternalType = ExternalAuthTranslationTraits->Type;
 	}
 
-	// todo: handle SystemAuthCredentialsOptions
+	if (!InitSystemAuthCredentialOptions(EOSAuthLoginOptions))
+	{
+		UE_LOG(LogOnlineServices, Warning, TEXT("FEOSAuthLoginOptions::Create: Failed - Failed to initialize system credential options"));
+		return TDefaultErrorResultInternal<FEOSAuthLoginOptions>(Errors::InvalidCreds());
+	}
 
 	// Translate scopes.
 	bool bAllScopesValid = true;
diff --git a/Engine/Plugins/Online/OnlineServicesEOSGS/Source/Private/Online/EOSAuthLoginOptionsCommon.cpp b/Engine/Plugins/Online/OnlineServicesEOSGS/Source/Private/Online/EOSAuthLoginOptionsCommon.cpp
new file mode 100644
--- /dev/null
+++ b/Engine/Plugins/Online/OnlineServicesEOSGS/Source/Private/Online/EOSAuthLoginOptionsCommon.cpp
@@ -0,0 +1,101 @@
+// Copyright Epic Games, Inc. All Rights Reserved.
+
+#include "Online/EOSAuthLoginOptionsCommon.h"
+
+#include "EOSShared.h"
+
+namespace UE::Online {
+
+const FEOSAuthTranslationTraits* FEOSAuthLoginOptionsCommon::GetLoginTranslatorTraits(FName Name)
+{
+	static const TMap<FName, FEOSAuthTranslationTraits> SupportedLoginTranslatorTraits = {
+		{ LoginCredentialsType::Password, { EOS_ELoginCredentialType::EOS_LCT_Password, EEOSAuthTranslationFlags::SetId | EEOSAuthTranslationFlags::SetTokenFromString } },
+		{ LoginCredentialsType::ExchangeCode, { EOS_ELoginCredentialType::EOS_LCT_ExchangeCode, EEOSAuthTranslationFlags::SetTokenFromString } },
+		{ LoginCredentialsType::PersistentAuth, { EOS_ELoginCredentialType::EOS_LCT_PersistentAuth, EEOSAuthTranslationFlags::None } },
+		{ LoginCredentialsType::Developer, { EOS_ELoginCredentialType::EOS_LCT_Developer, EEOSAuthTranslationFlags::SetId | EEOSAuthTranslationFlags::SetTokenFromString } },
+		{ LoginCredentialsType::RefreshToken, { EOS_ELoginCredentialType::EOS_LCT_RefreshToken, EEOSAuthTranslationFlags::SetTokenFromString } },
+		{ LoginCredentialsType::AccountPortal, { EOS_ELoginCredentialType::EOS_LCT_AccountPortal, EEOSAuthTranslationFlags::SetId | EEOSAuthTranslationFlags::SetTokenFromString } },
+		{ LoginCredentialsType::ExternalAuth, { EOS_ELoginCredentialType::EOS_LCT_ExternalAuth, EEOSAuthTranslationFlags::SetTokenFromExternalAuth } },
+	};
+
+	return SupportedLoginTranslatorTraits.Find(Name);
+}
+
+const FEOSExternalAuthTranslationTraits* FEOSAuthLoginOptionsCommon::GetExternalAuthTranslationTraits(FName ExternalAuthType)
+{
+	static const TMap<FName, FEOSExternalAuthTranslationTraits> SupportedExternalAuthTraits = {
+		{ ExternalLoginType::Epic, { EOS_EExternalCredentialType::EOS_ECT_EPIC } },
+		{ ExternalLoginType::SteamSessionTicket, { EOS_EExternalCredentialType::EOS_ECT_STEAM_SESSION_TICKET } },
+		{ ExternalLoginType::PsnIdToken, { EOS_EExternalCredentialType::EOS_ECT_PSN_ID_TOKEN } },
+		{ ExternalLoginType::XblXstsToken, { EOS_EExternalCredentialType::EOS_ECT_XBL_XSTS_TOKEN } },
+		{ ExternalLoginType::DiscordAccessToken, { EOS_EExternalCredentialType::EOS_ECT_DISCORD_ACCESS_TOKEN } },
+		{ ExternalLoginType::GogSessionTicket, { EOS_EExternalCredentialType::EOS_ECT_GOG_SESSION_TICKET } },
+		{ ExternalLoginType::NintendoIdToken, { EOS_EExternalCredentialType::EOS_ECT_NINTENDO_ID_TOKEN } },
+		{ ExternalLoginType::NintendoNsaIdToken, { EOS_EExternalCredentialType::EOS_ECT_NINTENDO_NSA_ID_TOKEN } },
+		{ ExternalLoginType::UplayAccessToken, { EOS_EExternalCredentialType::EOS_ECT_UPLAY_ACCESS_TOKEN } },
+		{ ExternalLoginType::OpenIdAccessToken, { EOS_EExternalCredentialType::EOS_ECT_OPENID_ACCESS_TOKEN } },
+		{ ExternalLoginType::DeviceIdAccessToken, { EOS_EExternalCredentialType::EOS_ECT_DEVICEID_ACCESS_TOKEN } },
+		{ ExternalLoginType::AppleIdToken, { EOS_EExternalCredentialType::EOS_ECT_APPLE_ID_TOKEN } },
+		{ ExternalLoginType::GoogleIdToken, { EOS_EExternalCredentialType::EOS_ECT_GOOGLE_ID_TOKEN } },
+		{ ExternalLoginType::OculusUserIdNonce, { EOS_EExternalCredentialType::EOS_ECT_OCULUS_USERID_NONCE } },
+		{ ExternalLoginType::ItchioJwt, { EOS_EExternalCredentialType::EOS_ECT_ITCHIO_JWT } },
+		{ ExternalLoginType::ItchioKey, { EOS_EExternalCredentialType::EOS_ECT_ITCHIO_KEY } },
+		{ ExternalLoginType::EpicIdToken, { EOS_EExternalCredentialType::EOS_ECT_EPIC_ID_TOKEN } },
+		{ ExternalLoginType::AmazonAccessToken, { EOS_EExternalCredentialType::EOS_ECT_AMAZON_ACCESS_TOKEN } },
+	};
+
+	return SupportedExternalAuthTraits.Find(ExternalAuthType);
+}
+
+FEOSAuthLoginOptionsCommon::FEOSAuthLoginOptionsCommon(FEOSAuthLoginOptionsCommon&& Other)
+{
+	*this = MoveTemp(Other);
+}
+
+FEOSAuthLoginOptionsCommon& FEOSAuthLoginOptionsCommon::operator=(FEOSAuthLoginOptionsCommon&& Other)
+{
+	CredentialsData = Other.CredentialsData;
+
+	// Pointer fixup.
+	if (CredentialsData.Id)
+	{
+		IdUtf8 = MoveTemp(Other.IdUtf8);
+		CredentialsData.Id = IdUtf8.GetData();
+	}
+	if (CredentialsData.Token)
+	{
+		TokenUtf8 = MoveTemp(Other.TokenUtf8);
+		CredentialsData.Token = TokenUtf8.GetData();
+	}
+	if (CredentialsData.SystemAuthCredentialsOptions)
+	{
+		// todo
+	}
+
+	Credentials = &CredentialsData;
+	ApiVersion = Other.ApiVersion;
+	ScopeFlags = Other.ScopeFlags;
+
+	Other.Credentials = nullptr;
+	return *this;
+}
+
+FEOSAuthLoginOptionsCommon::FEOSAuthLoginOptionsCommon()
+{
+	// EOS_Auth_LoginOptions init
+	UE_EOS_CHECK_API_MISMATCH(EOS_AUTH_LOGIN_API_LATEST, 3);
+	ApiVersion = 2;
+	Credentials = &CredentialsData;
+	ScopeFlags = EOS_EAuthScopeFlags::EOS_AS_NoFlags;
+	LoginFlags = 0;
+
+	// EOS_Auth_Credentials init
+	UE_EOS_CHECK_API_MISMATCH(EOS_AUTH_CREDENTIALS_API_LATEST, 4);
+	CredentialsData.ApiVersion = 4;
+	CredentialsData.Id = nullptr;
+	CredentialsData.Token = nullptr;
+	CredentialsData.Type = EOS_ELoginCredentialType::EOS_LCT_Password;
+	CredentialsData.SystemAuthCredentialsOptions = nullptr;
+	CredentialsData.ExternalType = EOS_EExternalCredentialType::EOS_ECT_EPIC;
+}
+}
\ No newline at end of file
diff --git a/Engine/Plugins/Online/OnlineServicesEOSGS/Source/Private/Online/EOSAuthLoginOptionsCommon.h b/Engine/Plugins/Online/OnlineServicesEOSGS/Source/Private/Online/EOSAuthLoginOptionsCommon.h
new file mode 100644
--- /dev/null
+++ b/Engine/Plugins/Online/OnlineServicesEOSGS/Source/Private/Online/EOSAuthLoginOptionsCommon.h
@@ -0,0 +1,52 @@
+// Copyright Epic Games, Inc. All Rights Reserved.
+
+#pragma once
+
+#include "Online/Auth.h"
+#include "Online/OnlineAsyncOp.h"
+
+#include "eos_auth_types.h"
+
+namespace UE::Online {
+
+enum class EEOSAuthTranslationFlags : uint8
+{
+	None = 0,
+	SetId = 1 << 0,
+	SetTokenFromString = 1 << 1,
+	SetTokenFromExternalAuth = 1 << 2,
+};
+ENUM_CLASS_FLAGS(EEOSAuthTranslationFlags);
+
+struct FEOSAuthTranslationTraits
+{
+	EOS_ELoginCredentialType Type;
+	EEOSAuthTranslationFlags Flags;
+};
+
+struct FEOSExternalAuthTranslationTraits
+{
+	EOS_EExternalCredentialType Type;
+};
+
+class FEOSAuthLoginOptionsCommon : public EOS_Auth_LoginOptions
+{
+public:
+	FEOSAuthLoginOptionsCommon(const FEOSAuthLoginOptionsCommon&) = delete;
+	FEOSAuthLoginOptionsCommon& operator=(const FEOSAuthLoginOptionsCommon&) = delete;
+	FEOSAuthLoginOptionsCommon(FEOSAuthLoginOptionsCommon&&);
+	FEOSAuthLoginOptionsCommon& operator=(FEOSAuthLoginOptionsCommon&&);
+
+protected:
+	FEOSAuthLoginOptionsCommon();
+
+	static const FEOSAuthTranslationTraits* GetLoginTranslatorTraits(FName Name);
+	static const FEOSExternalAuthTranslationTraits* GetExternalAuthTranslationTraits(FName ExternalAuthType);
+
+	static bool InitSystemAuthCredentialOptions(FEOSAuthLoginOptionsCommon& Options) { return true; }
+
+	EOS_Auth_Credentials CredentialsData;
+	TArray<char> IdUtf8;
+	TArray<char> TokenUtf8;
+};
+}
\ No newline at end of file
diff --git a/Engine/Plugins/Online/OnlineServicesEOSGS/Source/Private/Online/IOS/IOSEOSAuthLoginOptions.cpp b/Engine/Plugins/Online/OnlineServicesEOSGS/Source/Private/Online/IOS/IOSEOSAuthLoginOptions.cpp
new file mode 100644
--- /dev/null
+++ b/Engine/Plugins/Online/OnlineServicesEOSGS/Source/Private/Online/IOS/IOSEOSAuthLoginOptions.cpp
@@ -0,0 +1,62 @@
+// Copyright Epic Games, Inc. All Rights Reserved.
+
+#include "IOSEOSAuthLoginOptions.h"
+
+#include "EOSShared.h"
+#include "IOS/IOSAppDelegate.h"
+
+#import <AuthenticationServices/AuthenticationServices.h>
+
+@interface EOSAuthPresentationContext: NSObject <ASWebAuthenticationPresentationContextProviding>
+@end
+
+@implementation EOSAuthPresentationContext
+
+- (ASPresentationAnchor)presentationAnchorForWebAuthenticationSession:(ASWebAuthenticationSession *)session API_AVAILABLE(ios(13.0))
+{
+    return [[IOSAppDelegate GetDelegate] window];
+}
+
+@end
+
+namespace UE::Online
+{
+
+FIOSEOSAuthLoginOptions::FIOSEOSAuthLoginOptions(FIOSEOSAuthLoginOptions&& Other)
+{
+    *this = MoveTemp(Other);
+}
+
+FIOSEOSAuthLoginOptions& FIOSEOSAuthLoginOptions::operator=(FIOSEOSAuthLoginOptions&& Other)
+{
+    if (CredentialsOptions.PresentationContextProviding)
+	{
+        CFRelease((CFTypeRef*)CredentialsOptions.PresentationContextProviding);
+        CredentialsOptions = {};
+    }
+	
+    FEOSAuthLoginOptionsCommon::operator=(MoveTemp(Other));
+
+    if (Other.CredentialsData.SystemAuthCredentialsOptions)
+	{
+		CredentialsOptions.ApiVersion = Other.CredentialsOptions.ApiVersion;
+		CredentialsOptions.PresentationContextProviding = Other.CredentialsOptions.PresentationContextProviding;
+		CredentialsData.SystemAuthCredentialsOptions = &CredentialsOptions;
+	}
+	else
+	{
+		CredentialsData.SystemAuthCredentialsOptions = nullptr;
+	}
+	return *this;
+}
+
+bool FIOSEOSAuthLoginOptions::InitSystemAuthCredentialOptions(FIOSEOSAuthLoginOptions& Options)
+{
+	UE_EOS_CHECK_API_MISMATCH(EOS_IOS_AUTH_CREDENTIALSOPTIONS_API_LATEST, 1);    
+    Options.CredentialsOptions.ApiVersion = 1;
+	Options.CredentialsOptions.PresentationContextProviding = (void*)CFBridgingRetain([[EOSAuthPresentationContext alloc] init]);
+	Options.CredentialsData.SystemAuthCredentialsOptions = &Options.CredentialsOptions;
+    return true;
+}
+
+}
diff --git a/Engine/Plugins/Online/OnlineServicesEOSGS/Source/Private/Online/IOS/IOSEOSAuthLoginOptions.h b/Engine/Plugins/Online/OnlineServicesEOSGS/Source/Private/Online/IOS/IOSEOSAuthLoginOptions.h
new file mode 100644
--- /dev/null
+++ b/Engine/Plugins/Online/OnlineServicesEOSGS/Source/Private/Online/IOS/IOSEOSAuthLoginOptions.h
@@ -0,0 +1,28 @@
+// Copyright Epic Games, Inc. All Rights Reserved.
+
+#pragma once
+
+#include "Online/EOSAuthLoginOptionsCommon.h"
+
+#include "eos_ios.h"
+
+namespace UE::Online
+{
+
+class FIOSEOSAuthLoginOptions : public FEOSAuthLoginOptionsCommon
+{
+public:
+	FIOSEOSAuthLoginOptions(FIOSEOSAuthLoginOptions&&);
+	FIOSEOSAuthLoginOptions& operator=(FIOSEOSAuthLoginOptions&&);
+
+protected:
+	FIOSEOSAuthLoginOptions() = default;
+
+    static bool InitSystemAuthCredentialOptions(FIOSEOSAuthLoginOptions& Options);
+
+private:
+    EOS_IOS_Auth_CredentialsOptions CredentialsOptions = {};
+};
+
+using FPlatformEOSAuthLoginOptions = FIOSEOSAuthLoginOptions;
+}
