Fix missing include

diff --git a/Engine/Plugins/Online/OnlineServicesEOSGS/Source/Private/Online/EOSAuthLoginOptionsCommon.h b/Engine/Plugins/Online/OnlineServicesEOSGS/Source/Private/Online/EOSAuthLoginOptionsCommon.h
--- a/Engine/Plugins/Online/OnlineServicesEOSGS/Source/Private/Online/EOSAuthLoginOptionsCommon.h
+++ b/Engine/Plugins/Online/OnlineServicesEOSGS/Source/Private/Online/EOSAuthLoginOptionsCommon.h
@@ -5,6 +5,10 @@
 #include "Online/Auth.h"
 #include "Online/OnlineAsyncOp.h"
 
+#if defined(EOS_PLATFORM_BASE_FILE_NAME)
+#include EOS_PLATFORM_BASE_FILE_NAME
+#endif
+
 #include "eos_auth_types.h"
 
 namespace UE::Online {
