- Don't create a default platform during StartupModule

diff --git a/Engine/Plugins/Online/OnlineServicesEOSGS/Source/Private/Online/OnlineServicesEOSGSPlatformFactory.cpp b/Engine/Plugins/Online/OnlineServicesEOSGS/Source/Private/Online/OnlineServicesEOSGSPlatformFactory.cpp
--- a/Engine/Plugins/Online/OnlineServicesEOSGS/Source/Private/Online/OnlineServicesEOSGSPlatformFactory.cpp
+++ b/Engine/Plugins/Online/OnlineServicesEOSGS/Source/Private/Online/OnlineServicesEOSGSPlatformFactory.cpp
@@ -87,11 +87,6 @@ namespace UE::Online {
 
 FOnlineServicesEOSGSPlatformFactory::FOnlineServicesEOSGSPlatformFactory()
 {
-	// If a fork is requested, we need to wait for post-fork to create the default platform
-	if (!FForkProcessHelper::IsForkRequested() || FForkProcessHelper::IsForkedChildProcess())
-	{
-		GetDefaultPlatform();
-	}
 }
 
 FOnlineServicesEOSGSPlatformFactory& FOnlineServicesEOSGSPlatformFactory::Get()
