Setup OnlineServicesEOS and OnlineServicesEOSGS to use a platform allow list instead of a deny list so that platform specific functionality can be added.

diff --git a/Engine/Plugins/Online/OnlineServicesEOS/OnlineServicesEOS.uplugin b/Engine/Plugins/Online/OnlineServicesEOS/OnlineServicesEOS.uplugin
--- a/Engine/Plugins/Online/OnlineServicesEOS/OnlineServicesEOS.uplugin
+++ b/Engine/Plugins/Online/OnlineServicesEOS/OnlineServicesEOS.uplugin
@@ -15,9 +15,11 @@
 			"Name": "OnlineServicesEOS",
 			"Type": "Runtime",
 			"LoadingPhase": "PostConfigInit",
-			"PlatformDenyList": [
-				"LinuxArm64",
-				"TVOS"
+			"PlatformAllowList": [
+				"Win64",
+				"Mac",
+				"Linux",
+				"Android"
 			]
 		}
 	],
diff --git a/Engine/Plugins/Online/OnlineServicesEOSGS/OnlineServicesEOSGS.uplugin b/Engine/Plugins/Online/OnlineServicesEOSGS/OnlineServicesEOSGS.uplugin
--- a/Engine/Plugins/Online/OnlineServicesEOSGS/OnlineServicesEOSGS.uplugin
+++ b/Engine/Plugins/Online/OnlineServicesEOSGS/OnlineServicesEOSGS.uplugin
@@ -15,9 +15,11 @@
 			"Name": "OnlineServicesEOSGS",
 			"Type": "Runtime",
 			"LoadingPhase": "PostConfigInit",
-			"PlatformDenyList": [
-				"LinuxArm64",
-				"TVOS"
+			"PlatformAllowList": [
+				"Win64",
+				"Mac",
+				"Linux",
+				"Android"
 			]
 		}
 	],
