- Fix build error in OnlineServicesEOSGS due to EOSSDK upgrade

diff --git a/Engine/Plugins/Online/OnlineServicesEOSGS/Source/Private/Online/IOS/IOSEOSAuthLoginOptions.cpp b/Engine/Plugins/Online/OnlineServicesEOSGS/Source/Private/Online/IOS/IOSEOSAuthLoginOptions.cpp
--- a/Engine/Plugins/Online/OnlineServicesEOSGS/Source/Private/Online/IOS/IOSEOSAuthLoginOptions.cpp
+++ b/Engine/Plugins/Online/OnlineServicesEOSGS/Source/Private/Online/IOS/IOSEOSAuthLoginOptions.cpp
@@ -52,7 +52,7 @@ FIOSEOSAuthLoginOptions& FIOSEOSAuthLoginOptions::operator=(FIOSEOSAuthLoginOpti
 
 bool FIOSEOSAuthLoginOptions::InitSystemAuthCredentialOptions(FIOSEOSAuthLoginOptions& Options)
 {
-	UE_EOS_CHECK_API_MISMATCH(EOS_IOS_AUTH_CREDENTIALSOPTIONS_API_LATEST, 1);    
+	UE_EOS_CHECK_API_MISMATCH(EOS_IOS_AUTH_CREDENTIALSOPTIONS_API_LATEST, 2);
     Options.CredentialsOptions.ApiVersion = 1;
 	Options.CredentialsOptions.PresentationContextProviding = (void*)CFBridgingRetain([[EOSAuthPresentationContext alloc] init]);
 	Options.CredentialsData.SystemAuthCredentialsOptions = &Options.CredentialsOptions;
