Change OnlineServices EOS / EOSGS to not wrap auth errors in Errors::Unknown.

diff --git a/Engine/Plugins/Online/OnlineServicesEOS/Source/Private/Online/AuthEOS.cpp b/Engine/Plugins/Online/OnlineServicesEOS/Source/Private/Online/AuthEOS.cpp
--- a/Engine/Plugins/Online/OnlineServicesEOS/Source/Private/Online/AuthEOS.cpp
+++ b/Engine/Plugins/Online/OnlineServicesEOS/Source/Private/Online/AuthEOS.cpp
@@ -101,7 +101,7 @@ TOnlineAsyncOpHandle<FAuthLogin> FAuthEOS::Login(FAuthLogin::Params&& Params)
 				if (LoginResult.IsError())
 				{
 					UE_LOG(LogOnlineServices, Warning, TEXT("[FAuthEOS::Login] Failure: LoginEASImpl %s"), *LoginResult.GetErrorValue().GetLogString());
-					Op->SetError(Errors::Unknown(MoveTemp(LoginResult.GetErrorValue())));
+					Op->SetError(MoveTemp(LoginResult.GetErrorValue()));
 				}
 				else
 				{
diff --git a/Engine/Plugins/Online/OnlineServicesEOSGS/Source/Private/Online/AuthEOSGS.cpp b/Engine/Plugins/Online/OnlineServicesEOSGS/Source/Private/Online/AuthEOSGS.cpp
--- a/Engine/Plugins/Online/OnlineServicesEOSGS/Source/Private/Online/AuthEOSGS.cpp
+++ b/Engine/Plugins/Online/OnlineServicesEOSGS/Source/Private/Online/AuthEOSGS.cpp
@@ -308,6 +308,8 @@ TDefaultErrorResultInternal<FEOSAuthLoginOptions> FEOSAuthLoginOptions::CreateIm
 		FCStringAnsi::Strncpy(EOSAuthLoginOptions.TokenUtf8.GetData(), TCHAR_TO_UTF8(*ExternalAuthToken.Data), EOSAuthLoginOptions.TokenUtf8.Num());
 		EOSAuthLoginOptions.CredentialsData.Token = EOSAuthLoginOptions.TokenUtf8.GetData();
 		EOSAuthLoginOptions.CredentialsData.ExternalType = ExternalAuthTranslationTraits->Type;
+
+		UE_LOG(LogOnlineServices, VeryVerbose, TEXT("FEOSAuthLoginOptions::Create: Using token type: %s, Token data: %.*s"), *ExternalAuthToken.Type.ToString(), ExternalAuthToken.Data.Len(), *ExternalAuthToken.Data);
 	}
 
 	if (!InitSystemAuthCredentialOptions(EOSAuthLoginOptions))
@@ -505,7 +507,7 @@ TOnlineAsyncOpHandle<FAuthLogin> FAuthEOSGS::Login(FAuthLogin::Params&& Params)
 					if (LoginResult.IsError())
 					{
 						UE_LOG(LogOnlineServices, Warning, TEXT("[FAuthEOSGS::Login] Failure: LoginEASImpl %s"), *LoginResult.GetErrorValue().GetLogString());
-						Op->SetError(Errors::Unknown(MoveTemp(LoginResult.GetErrorValue())));
+						Op->SetError(MoveTemp(LoginResult.GetErrorValue()));
 					}
 					else
 					{
@@ -887,13 +889,13 @@ TFuture<TDefaultErrorResult<FAuthLoginEASImpl>> FAuthEOSGS::LoginEASImpl(const F
 		return MakeFulfilledPromise<TDefaultErrorResult<FAuthLoginEASImpl>>(MoveTemp(LoginOptionsResult.GetErrorValue())).GetFuture();
 	}
 
-	const bool IsPersistentAuthLogin = LoginParams.CredentialsType == LoginCredentialsType::PersistentAuth;
+	const bool bIsPersistentAuthLogin = LoginParams.CredentialsType == LoginCredentialsType::PersistentAuth;
 
 	TPromise<TDefaultErrorResult<FAuthLoginEASImpl>> Promise;
 	TFuture<TDefaultErrorResult<FAuthLoginEASImpl>> Future = Promise.GetFuture();
 
 	EOS_Async(EOS_Auth_Login, AuthHandle, MoveTemp(LoginOptionsResult.GetOkValue()),
-	[AuthHandle = AuthHandle, IsPersistentAuthLogin, Promise = MoveTemp(Promise)](const EOS_Auth_LoginCallbackInfo* Data) mutable -> void
+	[AuthHandle = AuthHandle, bIsPersistentAuthLogin, Promise = MoveTemp(Promise)](const EOS_Auth_LoginCallbackInfo* Data) mutable -> void
 	{
 		UE_LOG(LogOnlineServices, Verbose, TEXT("[FAuthEOSGS::LoginEASImpl] EOS_Auth_Login Result: [%s]"), *LexToString(Data->ResultCode));
 
@@ -933,7 +935,7 @@ TFuture<TDefaultErrorResult<FAuthLoginEASImpl>> FAuthEOSGS::LoginEASImpl(const F
 				Data->ResultCode == EOS_EResult::EOS_Auth_InvalidToken;
 
 			// Remove persistent auth credentials when they are found to be invalid.
-			if (IsPersistentAuthLogin && bShouldRemoveCachedToken)
+			if (bIsPersistentAuthLogin && bShouldRemoveCachedToken)
 			{
 				EOS_Auth_DeletePersistentAuthOptions DeletePersistentAuthOptions = {};
 				DeletePersistentAuthOptions.ApiVersion = 2;
