Removed unneeded captured variable shadowing class member

diff --git a/Engine/Plugins/Online/OnlineServicesEOSGS/Source/Private/Online/AuthEOSGS.cpp b/Engine/Plugins/Online/OnlineServicesEOSGS/Source/Private/Online/AuthEOSGS.cpp
--- a/Engine/Plugins/Online/OnlineServicesEOSGS/Source/Private/Online/AuthEOSGS.cpp
+++ b/Engine/Plugins/Online/OnlineServicesEOSGS/Source/Private/Online/AuthEOSGS.cpp
@@ -888,7 +888,7 @@ TFuture<TDefaultErrorResult<FAuthLoginEASImpl>> FAuthEOSGS::LoginEASImpl(const F
 	TFuture<TDefaultErrorResult<FAuthLoginEASImpl>> Future = Promise.GetFuture();
 
 	EOS_Async(EOS_Auth_Login, AuthHandle, MoveTemp(LoginOptionsResult.GetOkValue()),
-	[this, AuthHandle = AuthHandle, bIsPersistentAuthLogin, PlatformUserId = LoginParams.PlatformUserId, LinkAccountFlags, bAutoLinkAccount = LoginParams.bAutoLinkAccount, Promise = MoveTemp(Promise)](const EOS_Auth_LoginCallbackInfo* Data) mutable -> void
+	[this, bIsPersistentAuthLogin, PlatformUserId = LoginParams.PlatformUserId, LinkAccountFlags, bAutoLinkAccount = LoginParams.bAutoLinkAccount, Promise = MoveTemp(Promise)](const EOS_Auth_LoginCallbackInfo* Data) mutable -> void
 	{
 		UE_LOG(LogOnlineServices, Verbose, TEXT("[FAuthEOSGS::LoginEASImpl] EOS_Auth_Login Result: [%s]"), *LexToString(Data->ResultCode));
 
