- Add mobile support for EOSSDK PersistentAuth/AccountPortal login.

diff --git a/Engine/Plugins/Online/OnlineServicesEOSGS/Source/Private/Online/AuthEOSGS.cpp b/Engine/Plugins/Online/OnlineServicesEOSGS/Source/Private/Online/AuthEOSGS.cpp
--- a/Engine/Plugins/Online/OnlineServicesEOSGS/Source/Private/Online/AuthEOSGS.cpp
+++ b/Engine/Plugins/Online/OnlineServicesEOSGS/Source/Private/Online/AuthEOSGS.cpp
@@ -886,10 +886,33 @@ TFuture<TDefaultErrorResult<FAuthLoginEASImpl>> FAuthEOSGS::LoginEASImpl(const F
 
 	TPromise<TDefaultErrorResult<FAuthLoginEASImpl>> Promise;
 	TFuture<TDefaultErrorResult<FAuthLoginEASImpl>> Future = Promise.GetFuture();
+    
+	FEOSAuthLoginOptions LoginOptions = MoveTemp(LoginOptionsResult.GetOkValue());
+    
+#if PLATFORM_IOS
+	const EDeviceScreenOrientation OriginalDeviceOrientation = FPlatformMisc::GetAllowedDeviceOrientation();
+	// TODO move this into ASWebAuthenticationSession callbacks when they become available.
+	if (LoginOptions.Credentials->Type == EOS_ELoginCredentialType::EOS_LCT_AccountPortal)
+	{
+		FPlatformMisc::ForceUpdateDeviceOrientation(EDeviceScreenOrientation::Portrait);
+	}
+#endif
 
-	EOS_Async(EOS_Auth_Login, AuthHandle, MoveTemp(LoginOptionsResult.GetOkValue()),
-	[this, bIsPersistentAuthLogin, PlatformUserId = LoginParams.PlatformUserId, LinkAccountFlags, bAutoLinkAccount = LoginParams.bAutoLinkAccount, Promise = MoveTemp(Promise)](const EOS_Auth_LoginCallbackInfo* Data) mutable -> void
+	EOS_Async(EOS_Auth_Login, AuthHandle, MoveTemp(LoginOptions),
+	[this, bIsPersistentAuthLogin, PlatformUserId = LoginParams.PlatformUserId, LinkAccountFlags, bAutoLinkAccount = LoginParams.bAutoLinkAccount, Promise = MoveTemp(Promise)
+#if PLATFORM_IOS
+		, OriginalDeviceOrientation
+#endif
+	](const EOS_Auth_LoginCallbackInfo* Data) mutable -> void
 	{
+#if PLATFORM_IOS
+        // TODO move this into ASWebAuthenticationSession callbacks when they become available.
+		if (FPlatformMisc::GetAllowedDeviceOrientation() != OriginalDeviceOrientation)
+		{
+			FPlatformMisc::ForceUpdateDeviceOrientation(OriginalDeviceOrientation);
+		}
+#endif
+        
 		UE_LOG(LogOnlineServices, Verbose, TEXT("[FAuthEOSGS::LoginEASImpl] EOS_Auth_Login Result: [%s]"), *LexToString(Data->ResultCode));
 
 		if (Data->ResultCode == EOS_EResult::EOS_Success)
diff --git a/Engine/Source/Runtime/Core/Private/GenericPlatform/GenericPlatformMisc.cpp b/Engine/Source/Runtime/Core/Private/GenericPlatform/GenericPlatformMisc.cpp
--- a/Engine/Source/Runtime/Core/Private/GenericPlatform/GenericPlatformMisc.cpp
+++ b/Engine/Source/Runtime/Core/Private/GenericPlatform/GenericPlatformMisc.cpp
@@ -1799,6 +1799,11 @@ void FGenericPlatformMisc::SetAllowedDeviceOrientation(EDeviceScreenOrientation
 	AllowedDeviceOrientation = NewAllowedDeviceOrientation;
 }
 
+void FGenericPlatformMisc::ForceUpdateDeviceOrientation(EDeviceScreenOrientation NewDeviceOrientation)
+{
+    SetAllowedDeviceOrientation(NewDeviceOrientation);
+}
+
 int32 FGenericPlatformMisc::GetDeviceVolume()
 {
 	return -1;
diff --git a/Engine/Source/Runtime/Core/Private/IOS/IOSPlatformMisc.cpp b/Engine/Source/Runtime/Core/Private/IOS/IOSPlatformMisc.cpp
--- a/Engine/Source/Runtime/Core/Private/IOS/IOSPlatformMisc.cpp
+++ b/Engine/Source/Runtime/Core/Private/IOS/IOSPlatformMisc.cpp
@@ -365,6 +365,16 @@ void FIOSPlatformMisc::SetAllowedDeviceOrientation(EDeviceScreenOrientation NewA
 #endif
 }
 
+void FIOSPlatformMisc::ForceUpdateDeviceOrientation(EDeviceScreenOrientation NewDeviceOrientation)
+{
+	SetAllowedDeviceOrientation(NewDeviceOrientation);
+
+	dispatch_async(dispatch_get_main_queue(), ^
+	{
+		[[IOSAppDelegate GetDelegate].IOSController setNeedsUpdateOfSupportedInterfaceOrientations];
+	});
+}
+
 bool FIOSPlatformMisc::HasPlatformFeature(const TCHAR* FeatureName)
 {
 	if (FCString::Stricmp(FeatureName, TEXT("Metal")) == 0)
diff --git a/Engine/Source/Runtime/Core/Public/GenericPlatform/GenericPlatformMisc.h b/Engine/Source/Runtime/Core/Public/GenericPlatform/GenericPlatformMisc.h
--- a/Engine/Source/Runtime/Core/Public/GenericPlatform/GenericPlatformMisc.h
+++ b/Engine/Source/Runtime/Core/Public/GenericPlatform/GenericPlatformMisc.h
@@ -1627,6 +1627,11 @@ public:
 	 * Change the allowed orientation of the device. 
 	 */
 	static CORE_API void SetAllowedDeviceOrientation(EDeviceScreenOrientation NewAllowedDeviceOrientation);
+    
+    /**
+     * Change the allowed orientation of the device, and immediately force the orientation to update.
+     */
+    static CORE_API void ForceUpdateDeviceOrientation(EDeviceScreenOrientation NewDeviceOrientation);
 
 	/**
 	 * Returns the device volume if the device is capable of returning that information.
diff --git a/Engine/Source/Runtime/Core/Public/IOS/IOSPlatformMisc.h b/Engine/Source/Runtime/Core/Public/IOS/IOSPlatformMisc.h
--- a/Engine/Source/Runtime/Core/Public/IOS/IOSPlatformMisc.h
+++ b/Engine/Source/Runtime/Core/Public/IOS/IOSPlatformMisc.h
@@ -60,6 +60,7 @@ struct CORE_API FIOSPlatformMisc : public FApplePlatformMisc
 	UE_DEPRECATED(5.1, "SetDeviceOrientation is deprecated. Use SetAllowedDeviceOrientation instead.")
 	static void SetDeviceOrientation(EDeviceScreenOrientation NewDeviceOrientation);
 	static void SetAllowedDeviceOrientation(EDeviceScreenOrientation NewAllowedDeviceOrientation);
+    static void ForceUpdateDeviceOrientation(EDeviceScreenOrientation NewDeviceOrientation);
 	static int32 GetDeviceVolume();
 	static void SetBrightness(float Brightness);
 	static float GetBrightness();
