- CVar to allow using Epic UserAuth token instead of Epic ID token to log in to EOS_Connect_Login

diff --git a/Engine/Plugins/Online/OnlineServicesEOSGS/Source/Private/Online/AuthEOSGS.cpp b/Engine/Plugins/Online/OnlineServicesEOSGS/Source/Private/Online/AuthEOSGS.cpp
--- a/Engine/Plugins/Online/OnlineServicesEOSGS/Source/Private/Online/AuthEOSGS.cpp
+++ b/Engine/Plugins/Online/OnlineServicesEOSGS/Source/Private/Online/AuthEOSGS.cpp
@@ -22,6 +22,12 @@ namespace UE::Online
 }
 #endif
 
+TAutoConsoleVariable<bool> CVarEnableExternalAuthUserAuthToken (
+	TEXT("Online.EOS.EnableExternalAuthUserAuthToken"),
+	false,
+	TEXT("Configures GetExternalAuthTokenImpl to use EOS_Auth_CopyUserAuthToken when true, and EOS_Auth_CopyIdToken when false."),
+	ECVF_Default);
+
 #include "eos_auth.h"
 #include "eos_connect.h"
 
@@ -1028,28 +1034,56 @@ TFuture<TDefaultErrorResult<FAuthLogoutEASImpl>> FAuthEOSGS::LogoutEASImpl(const
 
 TDefaultErrorResult<FAuthGetExternalAuthTokenImpl> FAuthEOSGS::GetExternalAuthTokenImpl(const FAuthGetExternalAuthTokenImpl::Params& Params)
 {
-	EOS_Auth_CopyIdTokenOptions CopyIdTokenOptions = {};
-	CopyIdTokenOptions.ApiVersion = 1;
-	CopyIdTokenOptions.AccountId = Params.EpicAccountId;
-	UE_EOS_CHECK_API_MISMATCH(EOS_AUTH_COPYIDTOKEN_API_LATEST, 1);
-
-	EOS_Auth_IdToken* IdToken = nullptr;
-	EOS_EResult Result = EOS_Auth_CopyIdToken(AuthHandle, &CopyIdTokenOptions, &IdToken);
-	if (Result == EOS_EResult::EOS_Success)
+	if (CVarEnableExternalAuthUserAuthToken.GetValueOnGameThread())
 	{
-		ON_SCOPE_EXIT
+		EOS_Auth_CopyUserAuthTokenOptions Options = {};
+		Options.ApiVersion = 1;
+		UE_EOS_CHECK_API_MISMATCH(EOS_AUTH_COPYUSERAUTHTOKEN_API_LATEST, 1);
+
+		EOS_Auth_Token* Token = nullptr;
+		const EOS_EResult Result = EOS_Auth_CopyUserAuthToken(AuthHandle, &Options, Params.EpicAccountId, &Token);
+		if (Result == EOS_EResult::EOS_Success)
 		{
-			EOS_Auth_IdToken_Release(IdToken);
-		};
+			ON_SCOPE_EXIT
+			{
+				EOS_Auth_Token_Release(Token);
+			};
 
-		FExternalAuthToken ExternalAuthToken;
-		ExternalAuthToken.Type = ExternalLoginType::EpicIdToken;
-		ExternalAuthToken.Data = UTF8_TO_TCHAR(IdToken->JsonWebToken);
-		return TDefaultErrorResult<FAuthGetExternalAuthTokenImpl>(FAuthGetExternalAuthTokenImpl::Result{ MoveTemp(ExternalAuthToken) });
+			FExternalAuthToken ExternalAuthToken;
+			ExternalAuthToken.Type = ExternalLoginType::Epic;
+			ExternalAuthToken.Data = UTF8_TO_TCHAR(Token->AccessToken);
+			return TDefaultErrorResult<FAuthGetExternalAuthTokenImpl>(FAuthGetExternalAuthTokenImpl::Result{ MoveTemp(ExternalAuthToken) });
+		}
+		else
+		{
+			return TDefaultErrorResult<FAuthGetExternalAuthTokenImpl>(Errors::FromEOSResult(Result));
+		}
 	}
 	else
 	{
-		return TDefaultErrorResult<FAuthGetExternalAuthTokenImpl>(Errors::FromEOSResult(Result));
+		EOS_Auth_CopyIdTokenOptions Options = {};
+		Options.ApiVersion = 1;
+		Options.AccountId = Params.EpicAccountId;
+		UE_EOS_CHECK_API_MISMATCH(EOS_AUTH_COPYIDTOKEN_API_LATEST, 1);
+
+		EOS_Auth_IdToken* IdToken = nullptr;
+		const EOS_EResult Result = EOS_Auth_CopyIdToken(AuthHandle, &Options, &IdToken);
+		if (Result == EOS_EResult::EOS_Success)
+		{
+			ON_SCOPE_EXIT
+			{
+				EOS_Auth_IdToken_Release(IdToken);
+			};
+
+			FExternalAuthToken ExternalAuthToken;
+			ExternalAuthToken.Type = ExternalLoginType::EpicIdToken;
+			ExternalAuthToken.Data = UTF8_TO_TCHAR(IdToken->JsonWebToken);
+			return TDefaultErrorResult<FAuthGetExternalAuthTokenImpl>(FAuthGetExternalAuthTokenImpl::Result{ MoveTemp(ExternalAuthToken) });
+		}
+		else
+		{
+			return TDefaultErrorResult<FAuthGetExternalAuthTokenImpl>(Errors::FromEOSResult(Result));
+		}
 	}
 }
 
