Add OnlineServicesCompatibility method to help transitioning from OSSv1 to OSSv2.
Allow external UI visibility events to be captured by both OSS1 and OSSv2 simultaneously.
Hook up external UI events to AuthEOS and AuthEOSGS.

diff --git a/Engine/Plugins/Online/OnlineServicesEOS/Source/Private/Online/ExternalUIEOS.cpp b/Engine/Plugins/Online/OnlineServicesEOS/Source/Private/Online/ExternalUIEOS.cpp
--- a/Engine/Plugins/Online/OnlineServicesEOS/Source/Private/Online/ExternalUIEOS.cpp
+++ b/Engine/Plugins/Online/OnlineServicesEOS/Source/Private/Online/ExternalUIEOS.cpp
@@ -13,7 +13,7 @@
 namespace UE::Online {
 
 FExternalUIEOS::FExternalUIEOS(FOnlineServicesEOS& InServices)
-	: FExternalUICommon(InServices)
+	: Super(InServices)
 {
 }
 
diff --git a/Engine/Plugins/Online/OnlineServicesEOS/Source/Public/Online/ExternalUIEOS.h b/Engine/Plugins/Online/OnlineServicesEOS/Source/Public/Online/ExternalUIEOS.h
--- a/Engine/Plugins/Online/OnlineServicesEOS/Source/Public/Online/ExternalUIEOS.h
+++ b/Engine/Plugins/Online/OnlineServicesEOS/Source/Public/Online/ExternalUIEOS.h
@@ -3,7 +3,7 @@
 #pragma once
 
 #include "CoreMinimal.h"
-#include "Online/ExternalUICommon.h"
+#include "Online/ExternalUIEOSGS.h"
 
 #if defined(EOS_PLATFORM_BASE_FILE_NAME)
 #include EOS_PLATFORM_BASE_FILE_NAME
@@ -20,10 +20,10 @@ struct FExternalUIEOSConfig
 
 class FOnlineServicesEOS;
 
-class ONLINESERVICESEOS_API FExternalUIEOS : public FExternalUICommon
+class ONLINESERVICESEOS_API FExternalUIEOS : public FExternalUIEOSGS
 {
 public:
-	using Super = FExternalUICommon;
+	using Super = FExternalUIEOSGS;
 
 	FExternalUIEOS(FOnlineServicesEOS& InOwningSubsystem);
 	virtual void Initialize() override;
diff --git a/Engine/Plugins/Online/OnlineServicesEOSGS/Source/Private/Online/ExternalUIEOSGS.cpp b/Engine/Plugins/Online/OnlineServicesEOSGS/Source/Private/Online/ExternalUIEOSGS.cpp
new file mode 100644
--- /dev/null
+++ b/Engine/Plugins/Online/OnlineServicesEOSGS/Source/Private/Online/ExternalUIEOSGS.cpp
@@ -0,0 +1,88 @@
+// Copyright Epic Games, Inc. All Rights Reserved.
+
+#include "Online/ExternalUIEOSGS.h"
+
+#include "IEOSSDKManager.h"
+#include "Online/OnlineServicesEOSGS.h"
+
+#include "eos_ui.h"
+
+namespace UE::Online {
+
+FExternalUIEOSGS::FExternalUIEOSGS(FOnlineServicesEOSGS& InServices)
+	: Super(InServices)
+{
+}
+
+void FExternalUIEOSGS::Initialize()
+{
+	Super::Initialize();
+
+	IEOSSDKManager* SDKManager = IEOSSDKManager::Get();
+	check(SDKManager);
+
+	UIInterfaceHandle = EOS_Platform_GetUIInterface(*static_cast<FOnlineServicesEOSGS&>(GetServices()).GetEOSPlatformHandle());
+	check(UIInterfaceHandle != nullptr);
+
+	RegisterEventHandlers();
+}
+
+void FExternalUIEOSGS::PreShutdown()
+{
+	Super::PreShutdown();
+
+	UnregisterEventHandlers();
+
+	UIInterfaceHandle = nullptr;
+}
+
+void FExternalUIEOSGS::RegisterEventHandlers()
+{
+	// This delegate would cause a crash when running a dedicated server
+	if (!IsRunningDedicatedServer())
+	{
+		OnDisplaySettingsUpdated = EOS_RegisterComponentEventHandler(
+			this,
+			UIInterfaceHandle,
+			EOS_UI_ADDNOTIFYDISPLAYSETTINGSUPDATED_API_LATEST,
+			&EOS_UI_AddNotifyDisplaySettingsUpdated,
+			&EOS_UI_RemoveNotifyDisplaySettingsUpdated,
+			&FExternalUIEOSGS::HandleDisplaySettingsUpdated);
+		UE_EOS_CHECK_API_MISMATCH(EOS_UI_ADDNOTIFYDISPLAYSETTINGSUPDATED_API_LATEST, 1);
+	}
+}
+
+void FExternalUIEOSGS::UnregisterEventHandlers()
+{
+	OnDisplaySettingsUpdated = nullptr;
+}
+
+void FExternalUIEOSGS::HandleDisplaySettingsUpdated(const EOS_UI_OnDisplaySettingsUpdatedCallbackInfo* Data)
+{
+	check(Data);
+	FExternalUIProcessDisplaySettingsUpdatedImp::Params Params;
+	Params.bIsVisible = Data->bIsVisible == EOS_TRUE;
+	Params.bIsExclusiveInput = Data->bIsExclusiveInput == EOS_TRUE;
+	ProcessDisplaySettingsUpdatedImplOp(MoveTemp(Params));
+}
+
+TOnlineAsyncOpHandle<FExternalUIProcessDisplaySettingsUpdatedImp> FExternalUIEOSGS::ProcessDisplaySettingsUpdatedImplOp(FExternalUIProcessDisplaySettingsUpdatedImp::Params&& Params)
+{
+	TOnlineAsyncOpRef<FExternalUIProcessDisplaySettingsUpdatedImp> Op = GetOp<FExternalUIProcessDisplaySettingsUpdatedImp>(MoveTemp(Params));
+
+	Op->Then([this](TOnlineAsyncOp<FExternalUIProcessDisplaySettingsUpdatedImp>& InAsyncOp)
+	{
+		const FExternalUIProcessDisplaySettingsUpdatedImp::Params& Params = InAsyncOp.GetParams();
+
+		UE_LOG(LogOnlineServices, Log, TEXT("[FExternalUIEOSGS::ProcessDisplaySettingsUpdatedImplOp] Display settings changed notification received. bIsVisible: %s, bIsExclusiveInput: %s"), *::LexToString(Params.bIsVisible), *::LexToString(Params.bIsExclusiveInput));
+
+		OnExternalUIStatusChangedEvent.Broadcast(FExternalUIStatusChanged{ Params.bIsExclusiveInput });
+
+		InAsyncOp.SetResult(FExternalUIProcessDisplaySettingsUpdatedImp::Result{});
+	})
+	.Enqueue(GetSerialQueue());
+
+	return Op->GetHandle();
+}
+
+/* UE::Online */ }
diff --git a/Engine/Plugins/Online/OnlineServicesEOSGS/Source/Public/Online/ExternalUIEOSGS.h b/Engine/Plugins/Online/OnlineServicesEOSGS/Source/Public/Online/ExternalUIEOSGS.h
new file mode 100644
--- /dev/null
+++ b/Engine/Plugins/Online/OnlineServicesEOSGS/Source/Public/Online/ExternalUIEOSGS.h
@@ -0,0 +1,69 @@
+// Copyright Epic Games, Inc. All Rights Reserved.
+
+#pragma once
+
+#include "CoreMinimal.h"
+#include "Online/ExternalUICommon.h"
+#include "OnlineServicesEOSGSTypes.h"
+
+#include "eos_ui_types.h"
+
+namespace UE::Online {
+
+class FOnlineServicesEOSGS;
+
+struct FExternalUIProcessDisplaySettingsUpdatedImp
+{
+	static constexpr TCHAR Name[] = TEXT("ProcessDisplaySettingsUpdatedImp");
+
+	struct Params
+	{
+		/** True when any portion of the overlay is visible. */
+		bool bIsVisible = false;
+		/**
+		 * True when the overlay has switched to exclusive input mode.
+		 * While in exclusive input mode, no keyboard or mouse input will be sent to the game.
+		 */
+		bool bIsExclusiveInput = false;
+	};
+
+	struct Result
+	{
+	};
+};
+
+class ONLINESERVICESEOSGS_API FExternalUIEOSGS : public FExternalUICommon
+{
+public:
+	using Super = FExternalUICommon;
+
+	FExternalUIEOSGS(FOnlineServicesEOSGS& InOwningSubsystem);
+
+	virtual void Initialize() override;
+	virtual void PreShutdown() override;
+
+protected:
+	void RegisterEventHandlers();
+	void UnregisterEventHandlers();
+
+	void HandleDisplaySettingsUpdated(const EOS_UI_OnDisplaySettingsUpdatedCallbackInfo* Data);
+
+	TOnlineAsyncOpHandle<FExternalUIProcessDisplaySettingsUpdatedImp> ProcessDisplaySettingsUpdatedImplOp(FExternalUIProcessDisplaySettingsUpdatedImp::Params&& Params);
+
+	EOS_HUI UIInterfaceHandle = nullptr;
+	EOSEventRegistrationPtr OnDisplaySettingsUpdated;
+};
+
+namespace Meta {
+
+BEGIN_ONLINE_STRUCT_META(FExternalUIProcessDisplaySettingsUpdatedImp::Params)
+	ONLINE_STRUCT_FIELD(FExternalUIProcessDisplaySettingsUpdatedImp::Params, bIsVisible),
+	ONLINE_STRUCT_FIELD(FExternalUIProcessDisplaySettingsUpdatedImp::Params, bIsExclusiveInput)
+END_ONLINE_STRUCT_META()
+
+BEGIN_ONLINE_STRUCT_META(FExternalUIProcessDisplaySettingsUpdatedImp::Result)
+END_ONLINE_STRUCT_META()
+
+/* Meta*/ }
+
+/* UE::Online */ }
diff --git a/Engine/Plugins/Online/OnlineSubsystemUtils/Source/OnlineSubsystemUtils/Private/OnlineEngineInterfaceImpl.cpp b/Engine/Plugins/Online/OnlineSubsystemUtils/Source/OnlineSubsystemUtils/Private/OnlineEngineInterfaceImpl.cpp
--- a/Engine/Plugins/Online/OnlineSubsystemUtils/Source/OnlineSubsystemUtils/Private/OnlineEngineInterfaceImpl.cpp
+++ b/Engine/Plugins/Online/OnlineSubsystemUtils/Source/OnlineSubsystemUtils/Private/OnlineEngineInterfaceImpl.cpp
@@ -17,6 +17,12 @@ UOnlineEngineInterfaceImpl::UOnlineEngineInterfaceImpl(const FObjectInitializer&
 {
 }
 
+void UOnlineEngineInterfaceImpl::PostInitProperties()
+{
+	Super::PostInitProperties();
+	InitCompatibilityInterface();
+}
+
 bool UOnlineEngineInterfaceImpl::IsLoaded(FName OnlineIdentifier)
 {
 	return IOnlineSubsystem::IsLoaded(OnlineIdentifier);
@@ -556,6 +562,11 @@ void UOnlineEngineInterfaceImpl::BindToExternalUIOpening(const FOnlineExternalUI
 		OnExternalUIChangeDelegate.BindUObject(this, &ThisClass::OnExternalUIChange, Delegate);
 		Utils->SetEngineExternalUIBinding(OnExternalUIChangeDelegate);
 	}
+
+	if (UOnlineEngineInterface* CompatibilityInterface = OnlineServicesCompatibilityInterface.Get())
+	{
+		CompatibilityInterface->BindToExternalUIOpening(Delegate);
+	}
 }
 
 void UOnlineEngineInterfaceImpl::OnExternalUIChange(bool bInIsOpening, FOnlineExternalUIChanged Delegate)
@@ -678,3 +689,14 @@ void UOnlineEngineInterfaceImpl::OnPIELoginComplete(int32 LocalUserNum, bool bWa
 
 #endif
 
+void UOnlineEngineInterfaceImpl::InitCompatibilityInterface()
+{
+	if (bOnlineServicesCompatibilityEnabled)
+	{
+		const UClass* OnlineEngineInterfaceClass = StaticLoadClass(UOnlineEngineInterface::StaticClass(), NULL, TEXT("/Script/OnlineSubsystemUtils.OnlineServicesEngineInterfaceImpl"), NULL, LOAD_Quiet, NULL);
+		check(OnlineEngineInterfaceClass);
+
+		OnlineServicesCompatibilityInterface = NewObject<UOnlineEngineInterface>(GetTransientPackage(), OnlineEngineInterfaceClass);
+		check(OnlineServicesCompatibilityInterface);
+	}
+}
diff --git a/Engine/Plugins/Online/OnlineSubsystemUtils/Source/OnlineSubsystemUtils/Private/OnlineEngineInterfaceImpl.h b/Engine/Plugins/Online/OnlineSubsystemUtils/Source/OnlineSubsystemUtils/Private/OnlineEngineInterfaceImpl.h
--- a/Engine/Plugins/Online/OnlineSubsystemUtils/Source/OnlineSubsystemUtils/Private/OnlineEngineInterfaceImpl.h
+++ b/Engine/Plugins/Online/OnlineSubsystemUtils/Source/OnlineSubsystemUtils/Private/OnlineEngineInterfaceImpl.h
@@ -17,6 +17,9 @@ class ONLINESUBSYSTEMUTILS_API UOnlineEngineInterfaceImpl : public UOnlineEngine
 {
 	GENERATED_UCLASS_BODY()
 
+public:
+	virtual void PostInitProperties() override;
+
 	/**
 	 * Subsystem
 	 */
@@ -162,5 +165,15 @@ private:
 	/** Mapping of delegate handles for each online Login() call while in flight */
 	TMap<FName, FDelegateHandle> OnLoginPIECompleteDelegateHandlesForPIEInstances;
 	void OnPIELoginComplete(int32 LocalUserNum, bool bWasSuccessful, const FUniqueNetId& UserId, const FString& Error, FName OnlineIdentifier, FOnlineAutoLoginComplete InCompletionDelegate);
+
+private:
+	void InitCompatibilityInterface();
+
+	/** Whether to enable a compatibility interface for transitioning from OSSv1 to OSSv2. */
+	UPROPERTY(config)
+	bool bOnlineServicesCompatibilityEnabled = false;
+
+	UPROPERTY()
+	TObjectPtr<UOnlineEngineInterface> OnlineServicesCompatibilityInterface;
 };
 
