Make SetAllowedDeviceOrientation on iOS to trigger an update on current orientation if needed

diff --git a/Engine/Plugins/Online/OnlineServicesEOSGS/Source/Private/Online/AuthEOSGS.cpp b/Engine/Plugins/Online/OnlineServicesEOSGS/Source/Private/Online/AuthEOSGS.cpp
--- a/Engine/Plugins/Online/OnlineServicesEOSGS/Source/Private/Online/AuthEOSGS.cpp
+++ b/Engine/Plugins/Online/OnlineServicesEOSGS/Source/Private/Online/AuthEOSGS.cpp
@@ -901,7 +901,7 @@ TFuture<TDefaultErrorResult<FAuthLoginEASImpl>> FAuthEOSGS::LoginEASImpl(const F
 	// TODO move this into ASWebAuthenticationSession callbacks when they become available.
 	if (LoginOptions.Credentials->Type == EOS_ELoginCredentialType::EOS_LCT_AccountPortal)
 	{
-		FPlatformMisc::ForceUpdateDeviceOrientation(EDeviceScreenOrientation::Portrait);
+		FPlatformMisc::SetAllowedDeviceOrientation(EDeviceScreenOrientation::FullSensor);
 	}
 #endif
 
@@ -916,7 +916,7 @@ TFuture<TDefaultErrorResult<FAuthLoginEASImpl>> FAuthEOSGS::LoginEASImpl(const F
         // TODO move this into ASWebAuthenticationSession callbacks when they become available.
 		if (FPlatformMisc::GetAllowedDeviceOrientation() != OriginalDeviceOrientation)
 		{
-			FPlatformMisc::ForceUpdateDeviceOrientation(OriginalDeviceOrientation);
+			FPlatformMisc::SetAllowedDeviceOrientation(OriginalDeviceOrientation);
 		}
 #endif
         
diff --git a/Engine/Source/Runtime/Core/Private/GenericPlatform/GenericPlatformMisc.cpp b/Engine/Source/Runtime/Core/Private/GenericPlatform/GenericPlatformMisc.cpp
--- a/Engine/Source/Runtime/Core/Private/GenericPlatform/GenericPlatformMisc.cpp
+++ b/Engine/Source/Runtime/Core/Private/GenericPlatform/GenericPlatformMisc.cpp
@@ -1801,11 +1801,6 @@ void FGenericPlatformMisc::SetAllowedDeviceOrientation(EDeviceScreenOrientation
 	AllowedDeviceOrientation = NewAllowedDeviceOrientation;
 }
 
-void FGenericPlatformMisc::ForceUpdateDeviceOrientation(EDeviceScreenOrientation NewDeviceOrientation)
-{
-    FPlatformMisc::SetAllowedDeviceOrientation(NewDeviceOrientation);
-}
-
 int32 FGenericPlatformMisc::GetDeviceVolume()
 {
 	return -1;
diff --git a/Engine/Source/Runtime/Core/Private/IOS/IOSPlatformMisc.cpp b/Engine/Source/Runtime/Core/Private/IOS/IOSPlatformMisc.cpp
--- a/Engine/Source/Runtime/Core/Private/IOS/IOSPlatformMisc.cpp
+++ b/Engine/Source/Runtime/Core/Private/IOS/IOSPlatformMisc.cpp
@@ -359,16 +359,11 @@ void FIOSPlatformMisc::SetDeviceOrientation(EDeviceScreenOrientation NewDeviceOr
 void FIOSPlatformMisc::SetAllowedDeviceOrientation(EDeviceScreenOrientation NewAllowedDeviceOrientation)
 {
 	AllowedDeviceOrientation = NewAllowedDeviceOrientation;
-
+	
 #if !PLATFORM_TVOS && !PLATFORM_VISIONOS
 	[IOSAppDelegate GetDelegate].IOSView->SupportedInterfaceOrientations = GetUIInterfaceOrientationMask(NewAllowedDeviceOrientation);
 #endif
-}
-
-void FIOSPlatformMisc::ForceUpdateDeviceOrientation(EDeviceScreenOrientation NewDeviceOrientation)
-{
-	SetAllowedDeviceOrientation(NewDeviceOrientation);
-
+	
 	dispatch_async(dispatch_get_main_queue(), ^
 	{
 		[[IOSAppDelegate GetDelegate].IOSController setNeedsUpdateOfSupportedInterfaceOrientations];
diff --git a/Engine/Source/Runtime/Core/Public/GenericPlatform/GenericPlatformMisc.h b/Engine/Source/Runtime/Core/Public/GenericPlatform/GenericPlatformMisc.h
--- a/Engine/Source/Runtime/Core/Public/GenericPlatform/GenericPlatformMisc.h
+++ b/Engine/Source/Runtime/Core/Public/GenericPlatform/GenericPlatformMisc.h
@@ -1628,11 +1628,6 @@ public:
 	 */
 	static CORE_API void SetAllowedDeviceOrientation(EDeviceScreenOrientation NewAllowedDeviceOrientation);
     
-    /**
-     * Change the allowed orientation of the device, and immediately force the orientation to update.
-     */
-    static CORE_API void ForceUpdateDeviceOrientation(EDeviceScreenOrientation NewDeviceOrientation);
-
 	/**
 	 * Returns the device volume if the device is capable of returning that information.
 	 *  -1 : Unknown
diff --git a/Engine/Source/Runtime/Core/Public/IOS/IOSPlatformMisc.h b/Engine/Source/Runtime/Core/Public/IOS/IOSPlatformMisc.h
--- a/Engine/Source/Runtime/Core/Public/IOS/IOSPlatformMisc.h
+++ b/Engine/Source/Runtime/Core/Public/IOS/IOSPlatformMisc.h
@@ -60,7 +60,6 @@ struct CORE_API FIOSPlatformMisc : public FApplePlatformMisc
 	UE_DEPRECATED(5.1, "SetDeviceOrientation is deprecated. Use SetAllowedDeviceOrientation instead.")
 	static void SetDeviceOrientation(EDeviceScreenOrientation NewDeviceOrientation);
 	static void SetAllowedDeviceOrientation(EDeviceScreenOrientation NewAllowedDeviceOrientation);
-    static void ForceUpdateDeviceOrientation(EDeviceScreenOrientation NewDeviceOrientation);
 	static int32 GetDeviceVolume();
 	static void SetBrightness(float Brightness);
 	static float GetBrightness();
