- Fix NAME_None being stringified as "None" causing EOSGS PlatformFactory to skip the LoadEOSPlatformConfig and GetDefaultPlatformConfigName fallbacks

#rnx

diff --git a/Engine/Plugins/Online/OnlineServicesEOSGS/Source/Private/Online/OnlineServicesEOSGSPlatformFactory.cpp b/Engine/Plugins/Online/OnlineServicesEOSGS/Source/Private/Online/OnlineServicesEOSGSPlatformFactory.cpp
--- a/Engine/Plugins/Online/OnlineServicesEOSGS/Source/Private/Online/OnlineServicesEOSGSPlatformFactory.cpp
+++ b/Engine/Plugins/Online/OnlineServicesEOSGS/Source/Private/Online/OnlineServicesEOSGSPlatformFactory.cpp
@@ -120,7 +120,11 @@ IEOSPlatformHandlePtr FOnlineServicesEOSGSPlatformFactory::CreatePlatform(FName
 		return {};
 	}
 
-	FString PlatformConfigName = InstanceConfigName.ToString();
+	FString PlatformConfigName;
+	if (!InstanceConfigName.IsNone())
+	{
+		PlatformConfigName = InstanceConfigName.ToString();
+	}
 	if (PlatformConfigName.IsEmpty())
 	{
 		PlatformConfigName = LoadEOSPlatformConfig(SDKManager);
