- Ensure that EOSShared is loaded-before EOSGS, which is loaded before EOS, etc

diff --git a/Engine/Plugins/Online/OnlineServicesEOS/Source/Private/Online/OnlineServicesEOSModule.cpp b/Engine/Plugins/Online/OnlineServicesEOS/Source/Private/Online/OnlineServicesEOSModule.cpp
--- a/Engine/Plugins/Online/OnlineServicesEOS/Source/Private/Online/OnlineServicesEOSModule.cpp
+++ b/Engine/Plugins/Online/OnlineServicesEOS/Source/Private/Online/OnlineServicesEOSModule.cpp
@@ -29,11 +29,7 @@ int FOnlineServicesEOSModule::GetRegistryPriority()
 
 void FOnlineServicesEOSModule::StartupModule()
 {
-	const FName EOSSharedModuleName = TEXT("EOSShared");
-	if (!FModuleManager::Get().IsModuleLoaded(EOSSharedModuleName))
-	{
-		FModuleManager::Get().LoadModuleChecked(EOSSharedModuleName);
-	}
+	FModuleManager::Get().LoadModuleChecked(TEXT("OnlineServicesEOSGS"));
 
 	// Making sure we load the module at this point will avoid errors while cooking
 	const FName OnlineServicesInterfaceModuleName = TEXT("OnlineServicesInterface");
diff --git a/Engine/Plugins/Online/OnlineServicesEOSGS/Source/Private/Online/OnlineServicesEOSGSModule.cpp b/Engine/Plugins/Online/OnlineServicesEOSGS/Source/Private/Online/OnlineServicesEOSGSModule.cpp
--- a/Engine/Plugins/Online/OnlineServicesEOSGS/Source/Private/Online/OnlineServicesEOSGSModule.cpp
+++ b/Engine/Plugins/Online/OnlineServicesEOSGS/Source/Private/Online/OnlineServicesEOSGSModule.cpp
@@ -29,12 +29,8 @@ int FOnlineServicesEOSGSModule::GetRegistryPriority()
 
 void FOnlineServicesEOSGSModule::StartupModule()
 {
-	// Making sure we load the module at this point will avoid errors while cooking
-	const FName OnlineServicesInterfaceModuleName = TEXT("OnlineServicesInterface");
-	if (!FModuleManager::Get().IsModuleLoaded(OnlineServicesInterfaceModuleName))
-	{
-		FModuleManager::Get().LoadModuleChecked(OnlineServicesInterfaceModuleName);
-	}
+	FModuleManager::Get().LoadModuleChecked(TEXT("OnlineServicesInterface"));
+	FModuleManager::Get().LoadModuleChecked(TEXT("EOSShared"));
 
 	FOnlineServicesRegistry::Get().RegisterServicesFactory(EOnlineServices::Epic, MakeUnique<FOnlineServicesFactoryEOSGS>(), GetRegistryPriority());
 	FOnlineIdRegistryRegistry::Get().RegisterAccountIdRegistry(EOnlineServices::Epic, &FOnlineAccountIdRegistryEOSGS::Get(), GetRegistryPriority());
