Added parameter to lobby data retrieval method in OnlineServices to avoid making unnecessary queries.

diff --git a/Engine/Plugins/Online/OnlineServicesEOSGS/Source/Private/Online/LobbiesEOSGS.cpp b/Engine/Plugins/Online/OnlineServicesEOSGS/Source/Private/Online/LobbiesEOSGS.cpp
--- a/Engine/Plugins/Online/OnlineServicesEOSGS/Source/Private/Online/LobbiesEOSGS.cpp
+++ b/Engine/Plugins/Online/OnlineServicesEOSGS/Source/Private/Online/LobbiesEOSGS.cpp
@@ -2047,7 +2047,7 @@ TOnlineAsyncOpHandle<FLobbiesProcessLobbyNotificationImpl> FLobbiesEOSGS::Proces
 		InAsyncOp.Data.Set<TSharedRef<FLobbyDetailsEOS>>(UE_ONLINE_LOBBY_EOS_KEY_NAME_LOBBY_DETAILS, LobbyDetails.ToSharedRef());
 
 		// Fetch lobby snapshot. Fetching the snapshot resolves the account ids of all lobby members in the snapshot.
-		return LobbyDetails->GetLobbySnapshot();
+		return LobbyDetails->GetLobbySnapshot(true);
 	})
 	.Then([this](TOnlineAsyncOp<FLobbiesProcessLobbyNotificationImpl>& InAsyncOp, TDefaultErrorResultInternal<FLobbyServiceSnapshot>&& LobbySnapshotResult)
 	{
diff --git a/Engine/Plugins/Online/OnlineServicesEOSGS/Source/Private/Online/LobbiesEOSGSTypes.cpp b/Engine/Plugins/Online/OnlineServicesEOSGS/Source/Private/Online/LobbiesEOSGSTypes.cpp
--- a/Engine/Plugins/Online/OnlineServicesEOSGS/Source/Private/Online/LobbiesEOSGSTypes.cpp
+++ b/Engine/Plugins/Online/OnlineServicesEOSGS/Source/Private/Online/LobbiesEOSGSTypes.cpp
@@ -300,7 +300,7 @@ bool FLobbyDetailsEOS::IsBucketCompatible() const
 	return LobbyDetailsInfo->GetBucketId() == Prerequisites->BucketId;
 }
 
-TFuture<TDefaultErrorResultInternal<FLobbyServiceSnapshot>> FLobbyDetailsEOS::GetLobbySnapshot() const
+TFuture<TDefaultErrorResultInternal<FLobbyServiceSnapshot>> FLobbyDetailsEOS::GetLobbySnapshot(bool bCopyMemberData) const
 {
 	TSharedPtr<FAuthEOSGS> AuthInterface = Prerequisites->AuthInterface.Pin();
 	if (!AuthInterface)
@@ -310,31 +310,34 @@ TFuture<TDefaultErrorResultInternal<FLobbyServiceSnapshot>> FLobbyDetailsEOS::Ge
 		return MakeFulfilledPromise<TDefaultErrorResultInternal<FLobbyServiceSnapshot>>(Errors::MissingInterface()).GetFuture();
 	}
 
-	EOS_LobbyDetails_GetMemberCountOptions GetMemberCountOptions = {};
-	GetMemberCountOptions.ApiVersion = 1;
-	UE_EOS_CHECK_API_MISMATCH(EOS_LOBBYDETAILS_GETMEMBERCOUNT_API_LATEST, 1);
-
-	const uint32_t MemberCount = EOS_LobbyDetails_GetMemberCount(LobbyDetailsHandle, &GetMemberCountOptions);
-
 	TSharedRef<TArray<EOS_ProductUserId>> MemberProductUserIds = MakeShared<TArray<EOS_ProductUserId>>();
-	MemberProductUserIds->Reserve(MemberCount);
 
-	for (uint32_t MemberIndex = 0; MemberIndex < MemberCount; ++MemberIndex)
+	if(bCopyMemberData)
 	{
-		EOS_LobbyDetails_GetMemberByIndexOptions GetMemberByIndexOptions = {};
-		GetMemberByIndexOptions.ApiVersion = 1;
-		UE_EOS_CHECK_API_MISMATCH(EOS_LOBBYDETAILS_GETMEMBERBYINDEX_API_LATEST, 1);
-		GetMemberByIndexOptions.MemberIndex = MemberIndex;
-
-		// This method will return null if the querying user is not part of the lobby, in which case we can stop iterating
-		EOS_ProductUserId MemberId = EOS_LobbyDetails_GetMemberByIndex(LobbyDetailsHandle, &GetMemberByIndexOptions);
-		if (MemberId != nullptr)
-		{
-			MemberProductUserIds->Emplace(MemberId);
-		}
-		else
+		EOS_LobbyDetails_GetMemberCountOptions GetMemberCountOptions = {};
+		GetMemberCountOptions.ApiVersion = 1;
+		UE_EOS_CHECK_API_MISMATCH(EOS_LOBBYDETAILS_GETMEMBERCOUNT_API_LATEST, 1);
+
+		const uint32_t MemberCount = EOS_LobbyDetails_GetMemberCount(LobbyDetailsHandle, &GetMemberCountOptions);
+		MemberProductUserIds->Reserve(MemberCount);
+
+		for (uint32_t MemberIndex = 0; MemberIndex < MemberCount; ++MemberIndex)
 		{
-			break;
+			EOS_LobbyDetails_GetMemberByIndexOptions GetMemberByIndexOptions = {};
+			GetMemberByIndexOptions.ApiVersion = 1;
+			UE_EOS_CHECK_API_MISMATCH(EOS_LOBBYDETAILS_GETMEMBERBYINDEX_API_LATEST, 1);
+			GetMemberByIndexOptions.MemberIndex = MemberIndex;
+
+			// This method will return null if the querying user is not part of the lobby, in which case we can stop iterating
+			EOS_ProductUserId MemberId = EOS_LobbyDetails_GetMemberByIndex(LobbyDetailsHandle, &GetMemberByIndexOptions);
+			if (MemberId != nullptr)
+			{
+				MemberProductUserIds->Emplace(MemberId);
+			}
+			else
+			{
+				break;
+			}
 		}
 	}
 
@@ -729,7 +732,7 @@ TFuture<TDefaultErrorResultInternal<TSharedRef<FLobbyDataEOS>>> FLobbyDataEOS::C
 	TPromise<TDefaultErrorResultInternal<TSharedRef<FLobbyDataEOS>>> Promise;
 	TFuture<TDefaultErrorResultInternal<TSharedRef<FLobbyDataEOS>>> Future = Promise.GetFuture();
 	
-	LobbyDetails->GetLobbySnapshot()
+	LobbyDetails->GetLobbySnapshot(false)
 	.Next(
 	[
 		Promise = MoveTemp(Promise),
diff --git a/Engine/Plugins/Online/OnlineServicesEOSGS/Source/Private/Online/LobbiesEOSGSTypes.h b/Engine/Plugins/Online/OnlineServicesEOSGS/Source/Private/Online/LobbiesEOSGSTypes.h
--- a/Engine/Plugins/Online/OnlineServicesEOSGS/Source/Private/Online/LobbiesEOSGSTypes.h
+++ b/Engine/Plugins/Online/OnlineServicesEOSGS/Source/Private/Online/LobbiesEOSGSTypes.h
@@ -265,7 +265,7 @@ public:
 	 * Retrieve a lobby data snapshot from the EOS lobby details object.
 	 * The lobby schema will be used to translate attribute data before returning.
 	 */
-	TFuture<TDefaultErrorResultInternal<FLobbyServiceSnapshot>> GetLobbySnapshot() const;
+	TFuture<TDefaultErrorResultInternal<FLobbyServiceSnapshot>> GetLobbySnapshot(bool bCopyMemberData) const;
 
 	/**
 	 * Retrieve lobby member data snapshot from the EOS lobby details object.
