- Remove EOSShared module load from FOnlineServicesEOSGSPlatformFactory::CreatePlatform, this is already done in StartupModule

#rnx

diff --git a/Engine/Plugins/Online/OnlineServicesEOSGS/Source/Private/Online/OnlineServicesEOSGSPlatformFactory.cpp b/Engine/Plugins/Online/OnlineServicesEOSGS/Source/Private/Online/OnlineServicesEOSGSPlatformFactory.cpp
--- a/Engine/Plugins/Online/OnlineServicesEOSGS/Source/Private/Online/OnlineServicesEOSGSPlatformFactory.cpp
+++ b/Engine/Plugins/Online/OnlineServicesEOSGS/Source/Private/Online/OnlineServicesEOSGSPlatformFactory.cpp
@@ -101,12 +101,6 @@ void FOnlineServicesEOSGSPlatformFactory::TearDown()
 
 IEOSPlatformHandlePtr FOnlineServicesEOSGSPlatformFactory::CreatePlatform(FName InstanceName, FName InstanceConfigName)
 {
-	const FName EOSSharedModuleName = TEXT("EOSShared");
-	if (!FModuleManager::Get().IsModuleLoaded(EOSSharedModuleName))
-	{
-		FModuleManager::Get().LoadModuleChecked(EOSSharedModuleName);
-	}
-
 	IEOSSDKManager* const SDKManager = IEOSSDKManager::Get();
 	if (!SDKManager)
 	{
