- Add PreInit step to OnlineServicesEOS* instance creation, so we can return nullptr when we can't create a platform instance.
- Fix CreateServices not nullchecking the pointer returned from factory

#rnx

diff --git a/Engine/Plugins/Online/OnlineServices/Source/OnlineServicesInterface/Private/Online/OnlineServicesRegistry.cpp b/Engine/Plugins/Online/OnlineServices/Source/OnlineServicesInterface/Private/Online/OnlineServicesRegistry.cpp
--- a/Engine/Plugins/Online/OnlineServices/Source/OnlineServicesInterface/Private/Online/OnlineServicesRegistry.cpp
+++ b/Engine/Plugins/Online/OnlineServices/Source/OnlineServicesInterface/Private/Online/OnlineServicesRegistry.cpp
@@ -179,11 +179,13 @@ TSharedPtr<IOnlineServices> FOnlineServicesRegistry::CreateServices(EOnlineServi
 
 	TSharedPtr<IOnlineServices> Services;
 
-	FFactoryAndPriority* FactoryAndPriority = ServicesFactories.Find(OnlineServices);
-	if (FactoryAndPriority != nullptr)
+	if (FFactoryAndPriority* FactoryAndPriority = ServicesFactories.Find(OnlineServices))
 	{
 		Services = FactoryAndPriority->Factory->Create(InstanceName, InstanceConfigName);
-		Services->Init();
+		if (Services)
+		{
+			Services->Init();
+		}
 	}
 
 	return Services;
diff --git a/Engine/Plugins/Online/OnlineServicesEOS/Source/Private/Online/OnlineServicesEOSModule.cpp b/Engine/Plugins/Online/OnlineServicesEOS/Source/Private/Online/OnlineServicesEOSModule.cpp
--- a/Engine/Plugins/Online/OnlineServicesEOS/Source/Private/Online/OnlineServicesEOSModule.cpp
+++ b/Engine/Plugins/Online/OnlineServicesEOS/Source/Private/Online/OnlineServicesEOSModule.cpp
@@ -18,7 +18,12 @@ public:
 	virtual ~FOnlineServicesFactoryEOS() {}
 	virtual TSharedPtr<IOnlineServices> Create(FName InInstanceName, FName InInstanceConfigName) override
 	{
-		return MakeShared<FOnlineServicesEOS>(InInstanceName, InInstanceConfigName);
+		TSharedPtr<FOnlineServicesEOS> Result = MakeShared<FOnlineServicesEOS>(InInstanceName, InInstanceConfigName);
+		if (!Result->PreInit())
+		{
+			Result = nullptr;
+		}
+		return Result;
 	}
 };
 
diff --git a/Engine/Plugins/Online/OnlineServicesEOSGS/Source/Private/Online/OnlineServicesEOSGS.cpp b/Engine/Plugins/Online/OnlineServicesEOSGS/Source/Private/Online/OnlineServicesEOSGS.cpp
--- a/Engine/Plugins/Online/OnlineServicesEOSGS/Source/Private/Online/OnlineServicesEOSGS.cpp
+++ b/Engine/Plugins/Online/OnlineServicesEOSGS/Source/Private/Online/OnlineServicesEOSGS.cpp
@@ -28,30 +28,29 @@ FOnlineServicesEOSGS::FOnlineServicesEOSGS(FName InInstanceName, FName InInstanc
 {
 }
 
-void FOnlineServicesEOSGS::Init()
+bool FOnlineServicesEOSGS::PreInit()
 {
 	FOnlineServicesEOSGSPlatformFactory& PlatformFactory = FOnlineServicesEOSGSPlatformFactory::Get();
 	EOSPlatformHandle = PlatformFactory.CreatePlatform(InstanceName, InstanceConfigName);
-	if (EOSPlatformHandle)
+	if (!EOSPlatformHandle)
 	{
+		UE_LOG(LogOnlineServices, Warning, TEXT("[%hs] InstanceName=%s InstanceConfigName=%s EOSPlatformHandle=nullptr."), __FUNCTION__, *InstanceName.ToString(), *InstanceConfigName.ToString());
+		return false;
+	}
+	
 #if WITH_ENGINE
-		SocketSubsystem = MakeShared<FSocketSubsystemEOS>(EOSPlatformHandle, MakeShared< FSocketSubsystemEOSUtils_OnlineServicesEOS>(*this));
-		check(SocketSubsystem);
+	SocketSubsystem = MakeShared<FSocketSubsystemEOS>(EOSPlatformHandle, MakeShared< FSocketSubsystemEOSUtils_OnlineServicesEOS>(*this));
+	check(SocketSubsystem);
 
-		FString ErrorStr;
-		if (!SocketSubsystem->Init(ErrorStr))
-		{
-			UE_LOG(LogOnlineServices, Warning, TEXT("[FOnlineServicesEOSGS::Initialize] Unable to initialize Socket Subsystem. Error=[%s]"), *ErrorStr);
-		}
-#endif
-	}
-	else
+	FString ErrorStr;
+	if (!SocketSubsystem->Init(ErrorStr))
 	{
-		UE_LOG(LogOnlineServices, Verbose, TEXT("[FOnlineServicesEOSGS::Initialize] Unable to initialize Socket Subsystem. EOS Platform Handle was invalid."));
-		return;
+		UE_LOG(LogOnlineServices, Warning, TEXT("[%hs] Unable to initialize Socket Subsystem. Error=[%s]"), __FUNCTION__, *ErrorStr);
+		return false;
 	}
+#endif
 
-	Super::Init();
+	return true;
 }
 
 void FOnlineServicesEOSGS::Destroy()
diff --git a/Engine/Plugins/Online/OnlineServicesEOSGS/Source/Private/Online/OnlineServicesEOSGSModule.cpp b/Engine/Plugins/Online/OnlineServicesEOSGS/Source/Private/Online/OnlineServicesEOSGSModule.cpp
--- a/Engine/Plugins/Online/OnlineServicesEOSGS/Source/Private/Online/OnlineServicesEOSGSModule.cpp
+++ b/Engine/Plugins/Online/OnlineServicesEOSGS/Source/Private/Online/OnlineServicesEOSGSModule.cpp
@@ -18,7 +18,12 @@ public:
 	virtual ~FOnlineServicesFactoryEOSGS() {}
 	virtual TSharedPtr<IOnlineServices> Create(FName InInstanceName, FName InInstanceConfigName) override
 	{
-		return MakeShared<FOnlineServicesEOSGS>(InInstanceName, InInstanceConfigName);
+		TSharedPtr<FOnlineServicesEOSGS> Result = MakeShared<FOnlineServicesEOSGS>(InInstanceName, InInstanceConfigName);
+		if (!Result->PreInit())
+		{
+			Result = nullptr;
+		}
+		return Result;
 	}
 };
 
diff --git a/Engine/Plugins/Online/OnlineServicesEOSGS/Source/Public/Online/OnlineServicesEOSGS.h b/Engine/Plugins/Online/OnlineServicesEOSGS/Source/Public/Online/OnlineServicesEOSGS.h
--- a/Engine/Plugins/Online/OnlineServicesEOSGS/Source/Public/Online/OnlineServicesEOSGS.h
+++ b/Engine/Plugins/Online/OnlineServicesEOSGS/Source/Public/Online/OnlineServicesEOSGS.h
@@ -25,7 +25,6 @@ public:
 	FOnlineServicesEOSGS(FName InstanceName, FName InstanceConfigName);
 	virtual ~FOnlineServicesEOSGS() = default;
 
-	virtual void Init() override;
 	virtual void Destroy() override;
 	virtual void RegisterComponents() override;
 	virtual TOnlineResult<FGetResolvedConnectString> GetResolvedConnectString(FGetResolvedConnectString::Params&& Params) override;
@@ -36,6 +35,9 @@ public:
 	UE_DEPRECATED(5.5, "Please call GetServiceConfigNameStatic instead")
 	static const TCHAR* GetConfigNameStatic() { return GetServiceConfigNameStatic(); }
 	static const TCHAR* GetServiceConfigNameStatic() { return TEXT("EOS"); }
+
+	bool PreInit();
+
 protected:
 	IEOSPlatformHandlePtr EOSPlatformHandle;
 
