EOS Split Plugins 1/2: Engine changes

Deprecated OnlineIdEOSGS global methods
- Added new variants taking a Services parameter, to enable their use by other Services implementations.
- Modified relevant call-sites.

Added OnlineServicesEpicCommon plugin, module and classes to encapsulate common functionality between the EOS/EOSGS and EpicGame/EpicAccount Services

Added support for a new launch parameter -EnableEOSSplitPlugins, which will use EpicGame/EpicAccount instead of EOS/EOSGS
- Added split Auth component for the new services

diff --git a/Engine/Plugins/Online/EOSShared/Source/EOSShared/Public/EOSSharedTypes.h b/Engine/Plugins/Online/EOSShared/Source/EOSShared/Public/EOSSharedTypes.h
--- a/Engine/Plugins/Online/EOSShared/Source/EOSShared/Public/EOSSharedTypes.h
+++ b/Engine/Plugins/Online/EOSShared/Source/EOSShared/Public/EOSSharedTypes.h
@@ -2,7 +2,9 @@
 
 #pragma once
 
+#include "Async/Future.h"
 #include "CoreGlobals.h"
+#include "EOSShared.h"
 #include "Templates/SharedPointer.h"
 
 #if WITH_EOS_SDK
@@ -137,4 +139,264 @@ private:
 	}
 };
 
+namespace UE::Online {
+
+namespace Private {
+
+/** Class to handle all callbacks generically using a future to forward callback results */
+template<typename CallbackType>
+class TEOSCallback
+{
+	using CallbackFuncType = void (EOS_CALL*)(const CallbackType*);
+public:
+	TEOSCallback() = default;
+	TEOSCallback(TPromise<const CallbackType*>&& InPromise) : Promise(MoveTemp(InPromise)) {}
+	virtual ~TEOSCallback() = default;
+
+	operator CallbackFuncType()
+	{
+		return &CallbackImpl;
+	}
+
+	TFuture<const CallbackType*> GetFuture() { return Promise.GetFuture(); }
+
+private:
+	TPromise<const CallbackType*> Promise;
+
+	static void EOS_CALL CallbackImpl(const CallbackType* Data)
+	{
+		if (EOS_EResult_IsOperationComplete(Data->ResultCode) == EOS_FALSE)
+		{
+			// Ignore
+			return;
+		}
+		check(IsInGameThread());
+
+		TEOSCallback* CallbackThis = reinterpret_cast<TEOSCallback*>(Data->ClientData);
+		check(CallbackThis);
+
+		CallbackThis->Promise.EmplaceValue(Data);
+
+		delete CallbackThis;
+	}
+};
+
+template<typename Function> struct FEOSAsyncCallbackTraits;
+template<typename TCallbackData>
+struct FEOSAsyncCallbackTraits<void(*)(const TCallbackData*)>
+{
+	typedef TCallbackData CallbackDataType;
+};
+
+template<typename Function> struct FEOSAsyncTraits;
+template<typename TEOSHandle, typename TEOSParameters, typename TEOSCallback>
+struct FEOSAsyncTraits<void(*)(TEOSHandle, const TEOSParameters*, void*, const TEOSCallback)>
+{
+	typedef TEOSCallback CallbackType;
+	typedef typename FEOSAsyncCallbackTraits<TEOSCallback>::CallbackDataType CallbackDataType;
+};
+
+/* Private */ }
+
+/**
+ * This method requires a Promise that will be fulfilled by the EOS callback.
+ * This is primarily intended to be used with continuations that take a TPromise. e.g.:
+ *
+ *		InAsyncOp.Then([this, LoginOptions, Credentials](TOnlineAsyncOp<FAuthLogin>& InAsyncOp, TPromise<const EOS_Auth_LoginCallbackInfo*>&& Promise) mutable
+ *		{
+ *			LoginOptions.Credentials = &Credentials;
+ *			EOS_Async(EOS_Auth_Login, AuthHandle, LoginOptions, MoveTemp(Promise));
+ *		})
+ * 
+ * Note that when a continuation takes a Promise parameter, the Future for that Promise is already bound
+ * to the next continuation in the Op chain. Therefore, as soon as the EOS_Auth_Login callback is called,
+ * the promise will be fulfilled, and the next continuation will run immediately. This ensures that even
+ * if the SDK calls the callback immediately (i.e. before EOS_Async returns), the continuation is already
+ * bound, and we can safely consume the CallbackInfo. Note that the CallbackInfo is only valid for the
+ * duration of the callback, so it is not safe to bind a continuation _after_ calling the EOS method.
+ */
+// Warning: This signature will crash if the async method returns immediately.
+// The callback object is deleted inside the callback before EOS_Async returns the future back to the user.
+template<typename TEOSResult, typename TEOSHandle, typename TEOSParameters, typename TEOSFn>
+decltype(auto) EOS_Async(TEOSFn EOSFn, TEOSHandle EOSHandle, TEOSParameters Parameters, TPromise<const TEOSResult*>&& Promise)
+{
+	Private::TEOSCallback<TEOSResult>* Callback = new Private::TEOSCallback<TEOSResult>(Forward< TPromise<const TEOSResult*>>(Promise));
+	return EOSFn(EOSHandle, &Parameters, Callback, *Callback);
+}
+
+/**
+ * This method requires a Callback (lambda / function pointer) which will be called by the EOS callback
+ * with the CallbackInfo.
+ */
+template<typename TEOSHandle, typename TEOSParameters, typename TEOSFn, typename TCallback>
+decltype(auto) EOS_Async(TEOSFn EOSFn, TEOSHandle EOSHandle, const TEOSParameters& Parameters, TCallback&& Callback)
+{
+	using CallbackDataType = typename Private::FEOSAsyncTraits<TEOSFn>::CallbackDataType;
+	TUniquePtr<TCallback> CallbackObj = MakeUnique<TCallback>(Forward<TCallback>(Callback));
+	return EOSFn(EOSHandle, &Parameters, CallbackObj.Release(),
+	static_cast<void(EOS_CALL*)(const CallbackDataType* Data)>([](const CallbackDataType* Data)
+	{
+		if (EOS_EResult_IsOperationComplete(Data->ResultCode))
+		{
+			check(IsInGameThread());
+			TUniquePtr<TCallback> CallbackObj(reinterpret_cast<TCallback*>(Data->ClientData));
+			check(CallbackObj);
+			(*CallbackObj)(Data);
+		}
+	}));
+}
+
+class FEOSEventRegistration
+{
+public:
+	virtual ~FEOSEventRegistration() = default;
+};
+typedef TUniquePtr<FEOSEventRegistration> FEOSEventRegistrationPtr;
+
+using EOSEventRegistration UE_DEPRECATED(5.6, "EOSEventRegistration renamed to FEOSEventRegistration") = FEOSEventRegistration;
+using EOSEventRegistrationPtr UE_DEPRECATED(5.6, "EOSEventRegistrationPtr renamed to FEOSEventRegistrationPtr") = FEOSEventRegistrationPtr;
+
+/** 
+* EOS event registration utility for binding an EOS notifier registration to a RAII object which handles
+* unregistering when it exits scope. Intended to be used from a TOnlineComponent class.
+* 
+* Example:
+*	EOSEventRegistrationPtr OnLobbyUpdatedEOSEventRegistration = EOS_RegisterComponentEventHandler(
+*		this,
+*		LobbyHandle,
+*		EOS_LOBBY_ADDNOTIFYLOBBYUPDATERECEIVED_API_LATEST,
+*		&EOS_Lobby_AddNotifyLobbyUpdateReceived,
+*		&EOS_Lobby_RemoveNotifyLobbyUpdateReceived,
+*		&FLobbiesEOS::HandleLobbyUpdated);
+**/
+template <
+	typename ComponentHandlerClass,
+	typename EOSHandle,
+	typename EOSNotfyRegisterFunction,
+	typename EOSNotfyUnregisterFunction,
+	typename ComponentHandlerFunction>
+FEOSEventRegistrationPtr EOS_RegisterComponentEventHandler(
+	ComponentHandlerClass* HandlerClass,
+	EOSHandle ClientHandle,
+	int32_t ApiVersion,
+	EOSNotfyRegisterFunction NotfyRegisterFunction,
+	EOSNotfyUnregisterFunction NotfyUnregisterFunction,
+	ComponentHandlerFunction HandlerFunction);
+
+namespace Private {
+
+template<typename Function> struct TEOSCallbackTraitsBase;
+
+template<typename Component, typename EventData>
+struct TEOSCallbackTraitsBase<void (Component::*)(const EventData*)>
+{
+	typedef Component ComponentType;
+	typedef EventData EventDataType;
+};
+
+template<typename Function>
+struct TEOSCallbackTraits : public TEOSCallbackTraitsBase<typename std::remove_reference<Function>::type>
+{
+};
+
+template<typename Function> struct TEOSNotifyRegisterTraits;
+
+template<typename EOSHandle, typename Options, typename ClientData, typename NotificationFn, typename NotificationId>
+struct TEOSNotifyRegisterTraits<NotificationId (*)(EOSHandle, const Options*, ClientData*, NotificationFn)>
+{
+	typedef Options OptionsType;
+	typedef NotificationId NotificationIdType;
+};
+
+template <
+	typename ComponentHandlerClass,
+	typename EOSHandle,
+	typename EOSNotfyRegisterFunction,
+	typename EOSNotfyUnregisterFunction,
+	typename ComponentHandlerFunction>
+class FEOSEventRegistrationImpl : public FEOSEventRegistration
+{
+public:
+	typedef FEOSEventRegistrationImpl<
+		ComponentHandlerClass,
+		EOSHandle,
+		EOSNotfyRegisterFunction,
+		EOSNotfyUnregisterFunction,
+		ComponentHandlerFunction> ThisClass;
+
+	FEOSEventRegistrationImpl(
+		ComponentHandlerClass* HandlerClass,
+		EOSHandle ClientHandle,
+		int32_t ApiVersion,
+		EOSNotfyRegisterFunction NotfyRegisterFunction,
+		EOSNotfyUnregisterFunction NotfyUnregisterFunction,
+		ComponentHandlerFunction HandlerFunction)
+		: HandlerClass(HandlerClass)
+		, ClientHandle(ClientHandle)
+		, NotfyUnregisterFunction(NotfyUnregisterFunction)
+		, HandlerFunction(HandlerFunction)
+	{
+		using EventDataType = typename TEOSCallbackTraits<ComponentHandlerFunction>::EventDataType;
+		typename TEOSNotifyRegisterTraits<EOSNotfyRegisterFunction>::OptionsType Options = { };
+		Options.ApiVersion = ApiVersion;
+		NotificationId = NotfyRegisterFunction(ClientHandle, &Options, this,
+		static_cast<void(EOS_CALL*)(const EventDataType* Data)>([](const EventDataType* Data)
+		{
+			ThisClass* This = reinterpret_cast<ThisClass*>(Data->ClientData);
+			(This->HandlerClass->*This->HandlerFunction)(Data);
+		}));
+	}
+
+	virtual ~FEOSEventRegistrationImpl()
+	{
+		NotfyUnregisterFunction(ClientHandle, NotificationId);
+	}
+
+private:
+	FEOSEventRegistrationImpl() = delete;
+	FEOSEventRegistrationImpl(const FEOSEventRegistrationImpl&) = delete;
+	FEOSEventRegistrationImpl& operator=(const FEOSEventRegistrationImpl&) = delete;
+
+	typedef typename TEOSNotifyRegisterTraits<EOSNotfyRegisterFunction>::NotificationIdType NotificationIdType;
+
+	NotificationIdType NotificationId;
+	ComponentHandlerClass* HandlerClass;
+	EOSHandle ClientHandle;
+	EOSNotfyUnregisterFunction NotfyUnregisterFunction;
+	ComponentHandlerFunction HandlerFunction;
+};
+
+/* Private */ }
+
+template <
+	typename ComponentHandlerClass,
+	typename EOSHandle,
+	typename EOSNotfyRegisterFunction,
+	typename EOSNotfyUnregisterFunction,
+	typename ComponentHandlerFunction>
+FEOSEventRegistrationPtr EOS_RegisterComponentEventHandler(
+	ComponentHandlerClass* HandlerClass,
+	EOSHandle ClientHandle,
+	int32_t ApiVersion,
+	EOSNotfyRegisterFunction NotfyRegisterFunction,
+	EOSNotfyUnregisterFunction NotfyUnregisterFunction,
+	ComponentHandlerFunction HandlerFunction)
+{
+	return MakeUnique<
+		Private::FEOSEventRegistrationImpl<
+			ComponentHandlerClass,
+			EOSHandle,
+			EOSNotfyRegisterFunction,
+			EOSNotfyUnregisterFunction,
+			ComponentHandlerFunction>>(
+			HandlerClass,
+			ClientHandle,
+			ApiVersion,
+			NotfyRegisterFunction,
+			NotfyUnregisterFunction,
+			HandlerFunction);
+}
+
+/* UE::Online */ }
+
 #endif
\ No newline at end of file
diff --git a/Engine/Plugins/Online/OnlineServices/Source/OnlineServicesCommon/Private/Online/AuthCommon.cpp b/Engine/Plugins/Online/OnlineServices/Source/OnlineServicesCommon/Private/Online/AuthCommon.cpp
--- a/Engine/Plugins/Online/OnlineServices/Source/OnlineServicesCommon/Private/Online/AuthCommon.cpp
+++ b/Engine/Plugins/Online/OnlineServices/Source/OnlineServicesCommon/Private/Online/AuthCommon.cpp
@@ -85,6 +85,7 @@ void FAuthCommon::RegisterCommands()
 {
 	RegisterCommand(&FAuthCommon::Login);
 	RegisterCommand(&FAuthCommon::Logout);
+	RegisterCommand(&FAuthCommon::CreateAccount);
 	RegisterCommand(&FAuthCommon::LinkAccount);
 	RegisterCommand(&FAuthCommon::ModifyAccountAttributes);
 	RegisterCommand(&FAuthCommon::QueryExternalServerAuthTicket);
@@ -114,6 +115,13 @@ TOnlineAsyncOpHandle<FAuthLogout> FAuthCommon::Logout(FAuthLogout::Params&& Para
 	return Operation->GetHandle();
 }
 
+TOnlineAsyncOpHandle<FAuthCreateAccount> FAuthCommon::CreateAccount(FAuthCreateAccount::Params&& Params)
+{
+	TOnlineAsyncOpRef<FAuthCreateAccount> Operation = GetOp<FAuthCreateAccount>(MoveTemp(Params));
+	Operation->SetError(Errors::NotImplemented());
+	return Operation->GetHandle();
+}
+
 TOnlineAsyncOpHandle<FAuthLinkAccount> FAuthCommon::LinkAccount(FAuthLinkAccount::Params&& Params)
 {
 	TOnlineAsyncOpRef<FAuthLinkAccount> Operation = GetOp<FAuthLinkAccount>(MoveTemp(Params));
diff --git a/Engine/Plugins/Online/OnlineServices/Source/OnlineServicesCommon/Public/Online/AuthCommon.h b/Engine/Plugins/Online/OnlineServices/Source/OnlineServicesCommon/Public/Online/AuthCommon.h
--- a/Engine/Plugins/Online/OnlineServices/Source/OnlineServicesCommon/Public/Online/AuthCommon.h
+++ b/Engine/Plugins/Online/OnlineServices/Source/OnlineServicesCommon/Public/Online/AuthCommon.h
@@ -45,6 +45,7 @@ public:
 	// IAuth
 	virtual TOnlineAsyncOpHandle<FAuthLogin> Login(FAuthLogin::Params&& Params) override;
 	virtual TOnlineAsyncOpHandle<FAuthLogout> Logout(FAuthLogout::Params&& Params) override;
+	virtual TOnlineAsyncOpHandle<FAuthCreateAccount> CreateAccount(FAuthCreateAccount::Params&& Params) override;
 	virtual TOnlineAsyncOpHandle<FAuthLinkAccount> LinkAccount(FAuthLinkAccount::Params&& Params) override;
 	virtual TOnlineAsyncOpHandle<FAuthModifyAccountAttributes> ModifyAccountAttributes(FAuthModifyAccountAttributes::Params&& Params) override;
 	virtual TOnlineAsyncOpHandle<FAuthQueryExternalServerAuthTicket> QueryExternalServerAuthTicket(FAuthQueryExternalServerAuthTicket::Params&& Params) override;
diff --git a/Engine/Plugins/Online/OnlineServices/Source/OnlineServicesInterface/Public/Online/Auth.h b/Engine/Plugins/Online/OnlineServices/Source/OnlineServicesInterface/Public/Online/Auth.h
--- a/Engine/Plugins/Online/OnlineServices/Source/OnlineServicesInterface/Public/Online/Auth.h
+++ b/Engine/Plugins/Online/OnlineServices/Source/OnlineServicesInterface/Public/Online/Auth.h
@@ -192,6 +192,24 @@ struct FAuthLogout
 	};
 };
 
+struct FAuthCreateAccount
+{
+	static constexpr TCHAR Name[] = TEXT("CreateAccount");
+
+	struct Params
+	{
+		/** The PlatformUserId of the Local User making the request. */
+		FPlatformUserId PlatformUserId = PLATFORMUSERID_NONE;
+		/** The id associated with external auth data used to continue a previous login attempt. */
+		FLoginContinuationId ContinuationId;
+	};
+
+	struct Result
+	{
+		TSharedRef<const FAccountInfo> AccountInfo;
+	};
+};
+
 struct FAuthLinkAccount
 {
 	static constexpr TCHAR Name[] = TEXT("LinkAccount");
@@ -452,7 +470,7 @@ public:
 	 * 
 	 * Returns Errors::InvalidUser when no user has been associated with the provided credentials.
 	 * GetLinkAccountContinuationId can then be used to retrieve a continuance token which can be
-	 * used to call LinkAccount to associate credentials with the service.
+	 * passed to either CreateAccount or LinkAccount.
 	 */
 	virtual TOnlineAsyncOpHandle<FAuthLogin> Login(FAuthLogin::Params&& Params) = 0;
 
@@ -461,6 +479,11 @@ public:
 	 */
 	virtual TOnlineAsyncOpHandle<FAuthLogout> Logout(FAuthLogout::Params&& Params) = 0;
 
+	/**
+	* Create a new account and link it to an external auth method using a continuation token from a previous login attempt.
+	*/
+	virtual TOnlineAsyncOpHandle<FAuthCreateAccount> CreateAccount(FAuthCreateAccount::Params&& Params) = 0;
+
 	/**
 	 * Link a logged in account to an external auth method using a continuation token from a previous login attempt.
 	 */
@@ -631,6 +654,15 @@ END_ONLINE_STRUCT_META()
 BEGIN_ONLINE_STRUCT_META(FAuthLogout::Result)
 END_ONLINE_STRUCT_META()
 
+BEGIN_ONLINE_STRUCT_META(FAuthCreateAccount::Params)
+	ONLINE_STRUCT_FIELD(FAuthCreateAccount::Params, PlatformUserId),
+	ONLINE_STRUCT_FIELD(FAuthCreateAccount::Params, ContinuationId)
+END_ONLINE_STRUCT_META()
+
+BEGIN_ONLINE_STRUCT_META(FAuthCreateAccount::Result)
+	ONLINE_STRUCT_FIELD(FAuthCreateAccount::Result, AccountInfo)
+END_ONLINE_STRUCT_META()
+
 BEGIN_ONLINE_STRUCT_META(FAuthLinkAccount::Params)
 	ONLINE_STRUCT_FIELD(FAuthLinkAccount::Params, PlatformUserId),
 	ONLINE_STRUCT_FIELD(FAuthLinkAccount::Params, ContinuationId)
diff --git a/Engine/Plugins/Online/OnlineServicesEOS/Source/OnlineServicesEOS.Build.cs b/Engine/Plugins/Online/OnlineServicesEOS/Source/OnlineServicesEOS.Build.cs
--- a/Engine/Plugins/Online/OnlineServicesEOS/Source/OnlineServicesEOS.Build.cs
+++ b/Engine/Plugins/Online/OnlineServicesEOS/Source/OnlineServicesEOS.Build.cs
@@ -22,6 +22,7 @@ public class OnlineServicesEOS : ModuleRules
 				"CoreUObject",
 				"EOSSDK",
 				"EOSShared",
+				"OnlineServicesEpicCommon"
 			}
 		);
 	}
diff --git a/Engine/Plugins/Online/OnlineServicesEOS/Source/Private/Online/AuthEOS.cpp b/Engine/Plugins/Online/OnlineServicesEOS/Source/Private/Online/AuthEOS.cpp
--- a/Engine/Plugins/Online/OnlineServicesEOS/Source/Private/Online/AuthEOS.cpp
+++ b/Engine/Plugins/Online/OnlineServicesEOS/Source/Private/Online/AuthEOS.cpp
@@ -9,7 +9,7 @@
 #include "EOSShared.h"
 #include "IEOSSDKManager.h"
 #include "Online/AuthErrors.h"
-#include "Online/OnlineErrorEOSGS.h"
+#include "Online/OnlineErrorEpicCommon.h"
 #include "Online/OnlineIdEOS.h"
 #include "Online/OnlineServicesEOS.h"
 #include "Online/OnlineUtils.h"
@@ -51,7 +51,7 @@ namespace ELinkAccountTag
 const FName InternalAccount = TEXT("InternalAccount");
 }
 
-FAuthEOS::FAuthEOS(FOnlineServicesEOS& InServices)
+FAuthEOS::FAuthEOS(FOnlineServicesCommon& InServices)
 	: Super(InServices)
 {
 }
@@ -60,7 +60,7 @@ void FAuthEOS::Initialize()
 {
 	Super::Initialize();
 
-	UserInfoHandle = EOS_Platform_GetUserInfoInterface(*GetServices<FOnlineServicesEOS>().GetEOSPlatformHandle());
+	UserInfoHandle = EOS_Platform_GetUserInfoInterface(*GetServices<FOnlineServicesEpicCommon>().GetEOSPlatformHandle());
 	check(UserInfoHandle != nullptr);
 }
 
@@ -740,7 +740,7 @@ TFuture<TArray<FAccountId>> FAuthEOS::ResolveAccountIds(const FAccountId& LocalA
 
 	// If we failed to find all the handles, we need to query, which requires a valid LocalAccountId
 	// Note this is unavailable on Dedicated Servers as well, unlike EOS_Connect_QueryProductUserIdMappings
-	if (!ValidateOnlineId(LocalAccountId))
+	if (!FOnlineAccountIdRegistryEOS::ValidateOnlineId(LocalAccountId))
 	{
 		checkNoEntry();
 		return MakeFulfilledPromise<TArray<FAccountId>>().GetFuture();
@@ -843,7 +843,7 @@ TFuture<TArray<FAccountId>> FAuthEOS::ResolveAccountIds(const FAccountId& LocalA
 	}
 
 	// If we failed to find all the handles, we need to query, which requires a valid LocalAccountId
-	if (!(IsRunningDedicatedServer() || ValidateOnlineId(LocalAccountId)))
+	if (!(IsRunningDedicatedServer() || FOnlineAccountIdRegistryEOS::ValidateOnlineId(LocalAccountId)))
 	{
 		checkNoEntry();
 		return MakeFulfilledPromise<TArray<FAccountId>>().GetFuture();
@@ -953,4 +953,14 @@ FAccountId FAuthEOS::CreateAccountId(const EOS_EpicAccountId EpicAccountId, cons
 	return FOnlineAccountIdRegistryEOS::Get().FindOrAddAccountId(EpicAccountId, ProductUserId);
 }
 
+FAccountId FAuthEOS::FindAccountId(const EOS_ProductUserId ProductUserId)
+{
+	return UE::Online::FindAccountId(Services.GetServicesProvider(), ProductUserId);
+}
+
+FAccountId FAuthEOS::FindAccountId(const EOS_EpicAccountId EpicAccountId)
+{
+	return UE::Online::FindAccountId(Services.GetServicesProvider(), EpicAccountId);
+}
+
 /* UE::Online */ }
diff --git a/Engine/Plugins/Online/OnlineServicesEOS/Source/Private/Online/CommerceEOS.cpp b/Engine/Plugins/Online/OnlineServicesEOS/Source/Private/Online/CommerceEOS.cpp
--- a/Engine/Plugins/Online/OnlineServicesEOS/Source/Private/Online/CommerceEOS.cpp
+++ b/Engine/Plugins/Online/OnlineServicesEOS/Source/Private/Online/CommerceEOS.cpp
@@ -4,7 +4,7 @@
 
 #include "IEOSSDKManager.h"
 #include "Online/AuthEOS.h"
-#include "Online/OnlineErrorEOSGS.h"
+#include "Online/OnlineErrorEpicCommon.h"
 #include "Online/OnlineIdEOS.h"
 #include "Online/OnlineServicesEOS.h"
 
diff --git a/Engine/Plugins/Online/OnlineServicesEOS/Source/Private/Online/OnlineIdEOS.cpp b/Engine/Plugins/Online/OnlineServicesEOS/Source/Private/Online/OnlineIdEOS.cpp
--- a/Engine/Plugins/Online/OnlineServicesEOS/Source/Private/Online/OnlineIdEOS.cpp
+++ b/Engine/Plugins/Online/OnlineServicesEOS/Source/Private/Online/OnlineIdEOS.cpp
@@ -127,6 +127,16 @@ FAccountId FOnlineAccountIdRegistryEOS::FromReplicationData(const TArray<uint8>&
 	return FindOrAddAccountId(EpicAccountId, ProductUserId);
 }
 
+
+FAccountId FOnlineAccountIdRegistryEOS::FindOrAddAccountId(const EOS_ProductUserId ProductUserId)
+{
+	checkNoEntry();
+
+	UE_LOG(LogTemp, Warning, TEXT("This method is not implemented in FOnlineAccountIdRegistryEOS, use FindOrAddAccountId(const EOS_EpicAccountId InEpicAccountId, const EOS_ProductUserId InProductUserId) instead."));
+
+	return FAccountId();
+}
+
 FAccountId FOnlineAccountIdRegistryEOS::FindAccountId(const EOS_ProductUserId ProductUserId) const
 {
 	FAccountId AccountId;
@@ -270,7 +280,24 @@ FAccountId FOnlineAccountIdRegistryEOS::FindOrAddAccountId(const EOS_EpicAccount
 
 EOS_EpicAccountId GetEpicAccountId(const FAccountId& AccountId)
 {
-	return FOnlineAccountIdRegistryEOS::Get().GetAccountIdData(AccountId).EpicAccountId;
+	EOnlineServices ServicesType = AccountId.GetOnlineServicesType();
+
+	if (ServicesType == EOnlineServices::EpicAccount)
+	{
+		IOnlineAccountIdRegistry* Registry = FOnlineIdRegistryRegistry::Get().GetAccountIdRegistry(ServicesType);
+		check(Registry);
+
+		IOnlineAccountIdRegistryEpicAccount* EpicAccountRegistry = static_cast<IOnlineAccountIdRegistryEpicAccount*>(Registry);
+		check(EpicAccountRegistry);
+
+		return EpicAccountRegistry->GetEpicAccountId(AccountId);
+	}
+	else if (ServicesType == EOnlineServices::Epic)
+	{
+		return FOnlineAccountIdRegistryEOS::Get().GetAccountIdData(AccountId).EpicAccountId;
+	}
+
+	return nullptr;
 }
 
 EOS_EpicAccountId GetEpicAccountIdChecked(const FAccountId& AccountId)
@@ -282,12 +309,37 @@ EOS_EpicAccountId GetEpicAccountIdChecked(const FAccountId& AccountId)
 
 FAccountId FindAccountId(const EOS_EpicAccountId EpicAccountId)
 {
-	return FOnlineAccountIdRegistryEOS::Get().FindAccountId(EpicAccountId);
+	return FindAccountId(EOnlineServices::Epic, EpicAccountId);
+}
+
+FAccountId FindAccountId(EOnlineServices Services, const EOS_EpicAccountId EpicAccountId)
+{
+	if (Services == EOnlineServices::EpicAccount)
+	{
+		IOnlineAccountIdRegistry* Registry = FOnlineIdRegistryRegistry::Get().GetAccountIdRegistry(Services);
+		check(Registry);
+
+		IOnlineAccountIdRegistryEpicAccount* EpicAccountRegistry = static_cast<IOnlineAccountIdRegistryEpicAccount*>(Registry);
+		check(EpicAccountRegistry);
+
+		return EpicAccountRegistry->GetAccountId(EpicAccountId);
+	}
+	else if (Services == EOnlineServices::Epic)
+	{
+		return FOnlineAccountIdRegistryEOS::Get().FindAccountId(EpicAccountId);
+	}
+
+	return FAccountId();	
 }
 
 FAccountId FindAccountIdChecked(const EOS_EpicAccountId EpicAccountId)
 {
-	FAccountId Result = FindAccountId(EpicAccountId);
+	return FindAccountIdChecked(EOnlineServices::Epic, EpicAccountId);
+}
+
+FAccountId FindAccountIdChecked(EOnlineServices Services, const EOS_EpicAccountId EpicAccountId)
+{
+	FAccountId Result = FindAccountId(Services, EpicAccountId);
 	check(Result.IsValid());
 	return Result;
 }
diff --git a/Engine/Plugins/Online/OnlineServicesEOS/Source/Private/Online/OnlineIdEOS.h b/Engine/Plugins/Online/OnlineServicesEOS/Source/Private/Online/OnlineIdEOS.h
--- a/Engine/Plugins/Online/OnlineServicesEOS/Source/Private/Online/OnlineIdEOS.h
+++ b/Engine/Plugins/Online/OnlineServicesEOS/Source/Private/Online/OnlineIdEOS.h
@@ -35,6 +35,7 @@ public:
 	// End IOnlineAccountIdRegistry
 
 	// Begin IOnlineAccountRegistryEOSGS
+	virtual FAccountId FindOrAddAccountId(const EOS_ProductUserId ProductUserId) override;
 	virtual FAccountId FindAccountId(const EOS_ProductUserId ProductUserId) const override;
 	virtual EOS_ProductUserId GetProductUserId(const FAccountId& AccountId) const override;
 	// End IOnlineAccountRegistryEOSGS
@@ -47,6 +48,12 @@ public:
 
 	static FOnlineAccountIdRegistryEOS& Get();
 
+	template<typename IdType>
+	static inline bool ValidateOnlineId(const TOnlineId<IdType> OnlineId)
+	{
+		return OnlineId.IsValid() && (OnlineId.GetOnlineServicesType() == EOnlineServices::Epic);
+	}
+
 private:
 	friend class FAuthEOS;
 	friend class FOnlineServicesEOSModule;
diff --git a/Engine/Plugins/Online/OnlineServicesEOS/Source/Private/Online/OnlineServicesEOSModule.cpp b/Engine/Plugins/Online/OnlineServicesEOS/Source/Private/Online/OnlineServicesEOSModule.cpp
--- a/Engine/Plugins/Online/OnlineServicesEOS/Source/Private/Online/OnlineServicesEOSModule.cpp
+++ b/Engine/Plugins/Online/OnlineServicesEOS/Source/Private/Online/OnlineServicesEOSModule.cpp
@@ -2,6 +2,7 @@
 
 #include "Online/OnlineServicesEOSModule.h"
 
+#include "Misc/CommandLine.h"
 #include "Modules/ModuleManager.h"
 #include "Online/OnlineExecHandler.h"
 #include "Online/OnlineIdEOS.h"
@@ -34,6 +35,13 @@ int FOnlineServicesEOSModule::GetRegistryPriority()
 
 void FOnlineServicesEOSModule::StartupModule()
 {
+	const bool bEnableEOSSplitPlugins = FParse::Param(FCommandLine::Get(), TEXT("EnableEOSSplitPlugins"));
+	if (bEnableEOSSplitPlugins)
+	{
+		UE_LOG(LogOnlineServices, Verbose, TEXT("EnableEOSSplitPlugins launch parameter detected. Skipping FOnlineServicesEOSModule startup."));
+		return;
+	}
+
 	FModuleManager::Get().LoadModuleChecked(TEXT("OnlineServicesEOSGS"));
 
 	// Making sure we load the module at this point will avoid errors while cooking
diff --git a/Engine/Plugins/Online/OnlineServicesEOS/Source/Private/Online/PresenceEOS.cpp b/Engine/Plugins/Online/OnlineServicesEOS/Source/Private/Online/PresenceEOS.cpp
--- a/Engine/Plugins/Online/OnlineServicesEOS/Source/Private/Online/PresenceEOS.cpp
+++ b/Engine/Plugins/Online/OnlineServicesEOS/Source/Private/Online/PresenceEOS.cpp
@@ -5,7 +5,7 @@
 #include "EOSShared.h"
 #include "IEOSSDKManager.h"
 #include "Online/AuthEOS.h"
-#include "Online/OnlineErrorEOSGS.h"
+#include "Online/OnlineErrorEpicCommon.h"
 #include "Online/OnlineIdEOS.h"
 #include "Online/OnlineServicesEOS.h"
 
@@ -726,4 +726,9 @@ void FPresenceEOS::HandlePresenceChanged(const EOS_Presence_PresenceChangedCallb
 	}
 }
 
+FAccountId FPresenceEOS::FindAccountId(const EOS_EpicAccountId EpicAccountId)
+{
+	return UE::Online::FindAccountId(Services.GetServicesProvider(), EpicAccountId);
+}
+
 /* UE::Online */ }
diff --git a/Engine/Plugins/Online/OnlineServicesEOS/Source/Private/Online/SocialEOS.cpp b/Engine/Plugins/Online/OnlineServicesEOS/Source/Private/Online/SocialEOS.cpp
--- a/Engine/Plugins/Online/OnlineServicesEOS/Source/Private/Online/SocialEOS.cpp
+++ b/Engine/Plugins/Online/OnlineServicesEOS/Source/Private/Online/SocialEOS.cpp
@@ -4,7 +4,7 @@
 
 #include "IEOSSDKManager.h"
 #include "Online/AuthEOS.h"
-#include "Online/OnlineErrorEOSGS.h"
+#include "Online/OnlineErrorEpicCommon.h"
 #include "Online/OnlineIdEOS.h"
 #include "Online/OnlineServicesEOS.h"
 
@@ -436,4 +436,9 @@ void FSocialEOS::HandleFriendsUpdate(const EOS_Friends_OnFriendsUpdateInfo* Data
 	});
 }
 
+FAccountId FSocialEOS::FindAccountId(const EOS_EpicAccountId EpicAccountId)
+{
+	return UE::Online::FindAccountId(Services.GetServicesProvider(), EpicAccountId);
+}
+
 /* UE::Online */ }
diff --git a/Engine/Plugins/Online/OnlineServicesEOS/Source/Private/Online/UserInfoEOS.cpp b/Engine/Plugins/Online/OnlineServicesEOS/Source/Private/Online/UserInfoEOS.cpp
--- a/Engine/Plugins/Online/OnlineServicesEOS/Source/Private/Online/UserInfoEOS.cpp
+++ b/Engine/Plugins/Online/OnlineServicesEOS/Source/Private/Online/UserInfoEOS.cpp
@@ -7,7 +7,7 @@
 #include "Online/OnlineIdEOS.h"
 #include "Online/OnlineServicesEOS.h"
 #include "Online/AuthEOS.h"
-#include "Online/OnlineErrorEOSGS.h"
+#include "Online/OnlineErrorEpicCommon.h"
 
 #include "eos_userinfo.h"
 
diff --git a/Engine/Plugins/Online/OnlineServicesEOS/Source/Public/Online/AccountIdEOS.h b/Engine/Plugins/Online/OnlineServicesEOS/Source/Public/Online/AccountIdEOS.h
--- a/Engine/Plugins/Online/OnlineServicesEOS/Source/Public/Online/AccountIdEOS.h
+++ b/Engine/Plugins/Online/OnlineServicesEOS/Source/Public/Online/AccountIdEOS.h
@@ -9,11 +9,24 @@ typedef struct EOS_ProductUserIdDetails* EOS_ProductUserId;
 
 namespace UE::Online
 {
+	class IOnlineAccountIdRegistryEpicAccount : public IOnlineAccountIdRegistry
+	{
+	public:
+		virtual FAccountId GetAccountId(EOS_EpicAccountId EpicAccountId) = 0;
+		virtual EOS_EpicAccountId GetEpicAccountId(const FAccountId& AccountId) const = 0;
+	};
 
 ONLINESERVICESEOS_API EOS_EpicAccountId GetEpicAccountId(const FAccountId& AccountId);
 ONLINESERVICESEOS_API EOS_EpicAccountId GetEpicAccountIdChecked(const FAccountId& AccountId);
+
+UE_DEPRECATED(5.6, "This method is deprecated, please use the new version taking a EOnlineServices parameter")
 ONLINESERVICESEOS_API FAccountId FindAccountId(const EOS_EpicAccountId EpicAccountId);
+ONLINESERVICESEOS_API FAccountId FindAccountId(EOnlineServices Services, const EOS_EpicAccountId EpicAccountId);
+
+UE_DEPRECATED(5.6, "This method is deprecated, please use the new version taking a EOnlineServices parameter")
 ONLINESERVICESEOS_API FAccountId FindAccountIdChecked(const EOS_EpicAccountId EpicAccountId);
+ONLINESERVICESEOS_API FAccountId FindAccountIdChecked(EOnlineServices Services, const EOS_EpicAccountId EpicAccountId);
+
 ONLINESERVICESEOS_API FAccountId CreateAccountId(const EOS_EpicAccountId EpicAccountId, const EOS_ProductUserId ProductUserId);
 
 } /* namespace UE::Online */
diff --git a/Engine/Plugins/Online/OnlineServicesEOS/Source/Public/Online/AuthEOS.h b/Engine/Plugins/Online/OnlineServicesEOS/Source/Public/Online/AuthEOS.h
--- a/Engine/Plugins/Online/OnlineServicesEOS/Source/Public/Online/AuthEOS.h
+++ b/Engine/Plugins/Online/OnlineServicesEOS/Source/Public/Online/AuthEOS.h
@@ -21,7 +21,7 @@ class ONLINESERVICESEOS_API FAuthEOS : public FAuthEOSGS
 public:
 	using Super = FAuthEOSGS;
 
-	FAuthEOS(FOnlineServicesEOS& InOwningSubsystem);
+	FAuthEOS(FOnlineServicesCommon& InOwningSubsystem);
 	virtual ~FAuthEOS() = default;
 
 	// Begin IOnlineComponent
@@ -52,6 +52,9 @@ protected:
 	void ProcessSuccessfulLogin(TOnlineAsyncOp<FAuthLogin>& InAsyncOp);
 	void OnEASLoginStatusChanged(FAccountId LocalAccountId, ELoginStatus PreviousStatus, ELoginStatus CurrentStatus);
 
+	FAccountId FindAccountId(const EOS_ProductUserId ProductUserId);
+	FAccountId FindAccountId(const EOS_EpicAccountId EpicAccountId);
+
 	static FAccountId CreateAccountId(const EOS_EpicAccountId EpicAccountId, const EOS_ProductUserId ProductUserId);
 
 	EOS_HUserInfo UserInfoHandle = nullptr;
diff --git a/Engine/Plugins/Online/OnlineServicesEOS/Source/Public/Online/PresenceEOS.h b/Engine/Plugins/Online/OnlineServicesEOS/Source/Public/Online/PresenceEOS.h
--- a/Engine/Plugins/Online/OnlineServicesEOS/Source/Public/Online/PresenceEOS.h
+++ b/Engine/Plugins/Online/OnlineServicesEOS/Source/Public/Online/PresenceEOS.h
@@ -40,6 +40,8 @@ protected:
 
 	void HandlePresenceChanged(const EOS_Presence_PresenceChangedCallbackInfo* Data);
 
+	FAccountId FindAccountId(const EOS_EpicAccountId EpicAccountId);
+
 	EOS_HPresence PresenceHandle = nullptr;
 
 	/** Login status changed event handle */
diff --git a/Engine/Plugins/Online/OnlineServicesEOS/Source/Public/Online/SocialEOS.h b/Engine/Plugins/Online/OnlineServicesEOS/Source/Public/Online/SocialEOS.h
--- a/Engine/Plugins/Online/OnlineServicesEOS/Source/Public/Online/SocialEOS.h
+++ b/Engine/Plugins/Online/OnlineServicesEOS/Source/Public/Online/SocialEOS.h
@@ -36,6 +36,8 @@ public:
 protected:
 	void HandleFriendsUpdate(const EOS_Friends_OnFriendsUpdateInfo* Data);
 
+	FAccountId FindAccountId(const EOS_EpicAccountId EpicAccountId);
+
 	TMap<FAccountId, TMap<FAccountId, TSharedRef<FFriend>>> FriendsLists;
 	EOS_HFriends FriendsHandle = nullptr;
 	EOS_HUI UIHandle = nullptr;
diff --git a/Engine/Plugins/Online/OnlineServicesEOSGS/OnlineServicesEOSGS.uplugin b/Engine/Plugins/Online/OnlineServicesEOSGS/OnlineServicesEOSGS.uplugin
--- a/Engine/Plugins/Online/OnlineServicesEOSGS/OnlineServicesEOSGS.uplugin
+++ b/Engine/Plugins/Online/OnlineServicesEOSGS/OnlineServicesEOSGS.uplugin
@@ -31,6 +31,10 @@
 		}
 	],
 	"Plugins": [
+		{
+			"Name": "OnlineServicesEpicCommon",
+			"Enabled": true
+		},
 		{
 			"Name": "OnlineServices",
 			"Enabled": true
diff --git a/Engine/Plugins/Online/OnlineServicesEOSGS/Source/OnlineServicesEOSGS.Build.cs b/Engine/Plugins/Online/OnlineServicesEOSGS/Source/OnlineServicesEOSGS.Build.cs
--- a/Engine/Plugins/Online/OnlineServicesEOSGS/Source/OnlineServicesEOSGS.Build.cs
+++ b/Engine/Plugins/Online/OnlineServicesEOSGS/Source/OnlineServicesEOSGS.Build.cs
@@ -12,6 +12,7 @@ public class OnlineServicesEOSGS : ModuleRules
 				"Core",
 				"EOSSDK",
 				"OnlineServicesInterface",
+				"OnlineServicesEpicCommon",
 				"OnlineServicesCommon",
 				"OnlineServicesCommonEngineUtils",
 				"OnlineBase",
diff --git a/Engine/Plugins/Online/OnlineServicesEOSGS/Source/Private/Online/AchievementsEOSGS.cpp b/Engine/Plugins/Online/OnlineServicesEOSGS/Source/Private/Online/AchievementsEOSGS.cpp
--- a/Engine/Plugins/Online/OnlineServicesEOSGS/Source/Private/Online/AchievementsEOSGS.cpp
+++ b/Engine/Plugins/Online/OnlineServicesEOSGS/Source/Private/Online/AchievementsEOSGS.cpp
@@ -6,7 +6,7 @@
 #include "IEOSSDKManager.h"
 #include "Online/AchievementsErrors.h"
 #include "Online/AuthEOSGS.h"
-#include "Online/OnlineErrorEOSGS.h"
+#include "Online/OnlineErrorEpicCommon.h"
 #include "Online/OnlineIdEOSGS.h"
 #include "Online/OnlineServicesEOSGS.h"
 
@@ -432,4 +432,9 @@ void FAchievementsEOSGS::HandleAchievementsUnlocked(const EOS_Achievements_OnAch
 	}
 }
 
+FAccountId FAchievementsEOSGS::FindAccountId(const EOS_ProductUserId ProductUserId)
+{
+	return UE::Online::FindAccountId(Services.GetServicesProvider(), ProductUserId);
+}
+
 /* UE::Online */ }
diff --git a/Engine/Plugins/Online/OnlineServicesEOSGS/Source/Private/Online/AuthEOSGS.cpp b/Engine/Plugins/Online/OnlineServicesEOSGS/Source/Private/Online/AuthEOSGS.cpp
--- a/Engine/Plugins/Online/OnlineServicesEOSGS/Source/Private/Online/AuthEOSGS.cpp
+++ b/Engine/Plugins/Online/OnlineServicesEOSGS/Source/Private/Online/AuthEOSGS.cpp
@@ -8,7 +8,7 @@
 #include "Misc/CommandLine.h"
 #include "Online/AuthErrors.h"
 #include "Online/EOSAuthLoginOptionsCommon.h"
-#include "Online/OnlineErrorEOSGS.h"
+#include "Online/OnlineErrorEpicCommon.h"
 #include "Online/OnlineIdEOSGS.h"
 #include "Online/OnlineServicesEOSGS.h"
 #include "Online/OnlineUtils.h"
@@ -437,7 +437,7 @@ void FAccountInfoRegistryEOS::DoUnregister(const TSharedRef<FAccountInfo>& Accou
 	}
 }
 
-FAuthEOSGS::FAuthEOSGS(FOnlineServicesEOSGS& InServices)
+FAuthEOSGS::FAuthEOSGS(FOnlineServicesCommon& InServices)
 	: Super(InServices)
 {
 }
@@ -446,10 +446,10 @@ void FAuthEOSGS::Initialize()
 {
 	Super::Initialize();
 
-	AuthHandle = EOS_Platform_GetAuthInterface(*GetServices<FOnlineServicesEOSGS>().GetEOSPlatformHandle());
+	AuthHandle = EOS_Platform_GetAuthInterface(*GetServices<FOnlineServicesEpicCommon>().GetEOSPlatformHandle());
 	check(AuthHandle != nullptr);
 
-	ConnectHandle = EOS_Platform_GetConnectInterface(*GetServices<FOnlineServicesEOSGS>().GetEOSPlatformHandle());
+	ConnectHandle = EOS_Platform_GetConnectInterface(*GetServices<FOnlineServicesEpicCommon>().GetEOSPlatformHandle());
 	check(ConnectHandle != nullptr);
 
 	RegisterHandlers();
@@ -1514,7 +1514,7 @@ void FAuthEOSGS::InitializeConnectLoginRecoveryTimer(const TSharedRef<FAccountIn
 
 FAccountId FAuthEOSGS::CreateAccountId(const EOS_ProductUserId ProductUserId)
 {
-	return FOnlineAccountIdRegistryEOSGS::Get().FindOrAddAccountId(ProductUserId);
+	return FOnlineAccountIdRegistryEOSGS::GetRegistered(EOnlineServices::Epic).FindOrAddAccountId(ProductUserId);
 }
 
 FAuthEOSGS::FUserScopedData* FAuthEOSGS::GetUserScopedData(FPlatformUserId PlatformUserId)
diff --git a/Engine/Plugins/Online/OnlineServicesEOSGS/Source/Private/Online/EOSAuthLoginOptionsCommon.cpp b/Engine/Plugins/Online/OnlineServicesEOSGS/Source/Private/Online/EOSAuthLoginOptionsCommon.cpp
deleted file mode 100644
diff --git a/Engine/Plugins/Online/OnlineServicesEOSGS/Source/Private/Online/EOSAuthLoginOptionsCommon.h b/Engine/Plugins/Online/OnlineServicesEOSGS/Source/Private/Online/EOSAuthLoginOptionsCommon.h
deleted file mode 100644
diff --git a/Engine/Plugins/Online/OnlineServicesEOSGS/Source/Private/Online/LeaderboardsEOSGS.cpp b/Engine/Plugins/Online/OnlineServicesEOSGS/Source/Private/Online/LeaderboardsEOSGS.cpp
--- a/Engine/Plugins/Online/OnlineServicesEOSGS/Source/Private/Online/LeaderboardsEOSGS.cpp
+++ b/Engine/Plugins/Online/OnlineServicesEOSGS/Source/Private/Online/LeaderboardsEOSGS.cpp
@@ -5,7 +5,7 @@
 #include "EOSShared.h"
 #include "IEOSSDKManager.h"
 #include "Online/AuthEOSGS.h"
-#include "Online/OnlineErrorEOSGS.h"
+#include "Online/OnlineErrorEpicCommon.h"
 #include "Online/OnlineIdEOSGS.h"
 #include "Online/OnlineServicesEOSGS.h"
 
@@ -142,7 +142,9 @@ void QueryLeaderboardsEOS(EOS_HLeaderboards LeaderboardsHandle, const FAccountId
 	EOS_Async(EOS_Leaderboards_QueryLeaderboardRanks, LeaderboardsHandle, Options, MoveTemp(Promise));
 }
 
-void ReadEntriesInRange(EOS_HLeaderboards LeaderboardsHandle, uint32 StartIndex, uint32 EndIndex, TArray<FLeaderboardEntry>& OutEntries)
+}
+
+void FLeaderboardsEOSGS::ReadEntriesInRange(uint32 StartIndex, uint32 EndIndex, TArray<FLeaderboardEntry>& OutEntries)
 {
 	for (uint32 i = StartIndex; i <= EndIndex; ++i)
 	{
@@ -156,7 +158,7 @@ void ReadEntriesInRange(EOS_HLeaderboards LeaderboardsHandle, uint32 StartIndex,
 		if (CopyResult == EOS_EResult::EOS_Success)
 		{
 			FLeaderboardEntry& LeaderboardEntry = OutEntries.Emplace_GetRef();
-			LeaderboardEntry.AccountId = FindAccountIdChecked(LeaderboardRecord->UserId);
+			LeaderboardEntry.AccountId = FindAccountIdChecked(Services.GetServicesProvider(), LeaderboardRecord->UserId);
 			LeaderboardEntry.Rank = LeaderboardRecord->Rank;
 			LeaderboardEntry.Score = LeaderboardRecord->Score;
 
@@ -165,8 +167,6 @@ void ReadEntriesInRange(EOS_HLeaderboards LeaderboardsHandle, uint32 StartIndex,
 	}
 }
 
-}
-
 TOnlineAsyncOpHandle<FReadEntriesAroundRank> FLeaderboardsEOSGS::ReadEntriesAroundRank(FReadEntriesAroundRank::Params&& Params)
 {
 	TOnlineAsyncOpRef<FReadEntriesAroundRank> Op = GetJoinableOp<FReadEntriesAroundRank>(MoveTemp(Params));
@@ -204,7 +204,7 @@ TOnlineAsyncOpHandle<FReadEntriesAroundRank> FLeaderboardsEOSGS::ReadEntriesArou
 			int32 StartIndex = Params.Rank;
 			int32 EndIndex = StartIndex + Params.Limit;
 			EndIndex = FMath::Clamp(EndIndex, 0, EOS_LEADERBOARD_MAX_ENTRIES-1);
-			Private::ReadEntriesInRange(LeaderboardsHandle, StartIndex, EndIndex, Result.Entries);
+			ReadEntriesInRange(StartIndex, EndIndex, Result.Entries);
 			InAsyncOp.SetResult(MoveTemp(Result));
 		})
 		.Enqueue(GetSerialQueue());
@@ -259,7 +259,7 @@ TOnlineAsyncOpHandle<FReadEntriesAroundUser> FLeaderboardsEOSGS::ReadEntriesArou
 				StartIndex = FMath::Clamp(StartIndex, 0, EOS_LEADERBOARD_MAX_ENTRIES - 1);
 				int32 EndIndex = StartIndex + Params.Limit;
 				EndIndex = FMath::Clamp(EndIndex, 0, EOS_LEADERBOARD_MAX_ENTRIES - 1);
-				Private::ReadEntriesInRange(LeaderboardsHandle, StartIndex, EndIndex, Result.Entries);
+				ReadEntriesInRange(StartIndex, EndIndex, Result.Entries);
 
 				EOS_Leaderboards_LeaderboardRecord_Release(LeaderboardRecord);
 			}
diff --git a/Engine/Plugins/Online/OnlineServicesEOSGS/Source/Private/Online/LobbiesEOSGS.cpp b/Engine/Plugins/Online/OnlineServicesEOSGS/Source/Private/Online/LobbiesEOSGS.cpp
--- a/Engine/Plugins/Online/OnlineServicesEOSGS/Source/Private/Online/LobbiesEOSGS.cpp
+++ b/Engine/Plugins/Online/OnlineServicesEOSGS/Source/Private/Online/LobbiesEOSGS.cpp
@@ -5,7 +5,7 @@
 #include "EOSShared.h"
 #include "Online/AuthEOSGS.h"
 #include "Online/LobbiesEOSGSTypes.h"
-#include "Online/OnlineErrorEOSGS.h"
+#include "Online/OnlineErrorEpicCommon.h"
 #include "Online/OnlineIdEOSGS.h"
 #include "Online/OnlineServicesEOSGS.h"
 
@@ -2127,4 +2127,9 @@ TOnlineAsyncOpHandle<FLobbiesProcessLobbyNotificationImpl> FLobbiesEOSGS::Proces
 	return Op->GetHandle();
 }
 
+FAccountId FLobbiesEOSGS::FindAccountId(const EOS_ProductUserId ProductUserId)
+{
+	return UE::Online::FindAccountId(Services.GetServicesProvider(), ProductUserId);
+}
+
 /* UE::Online */ }
diff --git a/Engine/Plugins/Online/OnlineServicesEOSGS/Source/Private/Online/LobbiesEOSGSTypes.cpp b/Engine/Plugins/Online/OnlineServicesEOSGS/Source/Private/Online/LobbiesEOSGSTypes.cpp
--- a/Engine/Plugins/Online/OnlineServicesEOSGS/Source/Private/Online/LobbiesEOSGSTypes.cpp
+++ b/Engine/Plugins/Online/OnlineServicesEOSGS/Source/Private/Online/LobbiesEOSGSTypes.cpp
@@ -4,7 +4,7 @@
 
 #include "EOSShared.h"
 #include "Online/AuthEOSGS.h"
-#include "Online/OnlineErrorEOSGS.h"
+#include "Online/OnlineErrorEpicCommon.h"
 #include "Online/OnlineIdEOSGS.h"
 #include "Online/OnlineUtils.h"
 #include <eos_lobby.h>
@@ -922,7 +922,7 @@ TFuture<TDefaultErrorResultInternal<TSharedRef<FLobbyInviteDataEOS>>> FLobbyInvi
 		}
 
 		// Once the lobby data has been resolved the FAccountId for the sender is expected to be in the AccountID cache.
-		const FAccountId SenderAccountId = FindAccountId(Sender);
+		const FAccountId SenderAccountId = FindAccountId(LocalAccountId.GetOnlineServicesType(), Sender);
 		if (!SenderAccountId.IsValid())
 		{
 			UE_LOG(LogOnlineServices, Warning, TEXT("[FLobbyInviteDataEOS::CreateFromInviteId] Failed: Unable to find account id for sender. User[%s], InviteId[%s], Lobby[%s]"),
diff --git a/Engine/Plugins/Online/OnlineServicesEOSGS/Source/Private/Online/OnlineErrorEOSGS.cpp b/Engine/Plugins/Online/OnlineServicesEOSGS/Source/Private/Online/OnlineErrorEOSGS.cpp
deleted file mode 100644
diff --git a/Engine/Plugins/Online/OnlineServicesEOSGS/Source/Private/Online/OnlineIdEOSGS.cpp b/Engine/Plugins/Online/OnlineServicesEOSGS/Source/Private/Online/OnlineIdEOSGS.cpp
--- a/Engine/Plugins/Online/OnlineServicesEOSGS/Source/Private/Online/OnlineIdEOSGS.cpp
+++ b/Engine/Plugins/Online/OnlineServicesEOSGS/Source/Private/Online/OnlineIdEOSGS.cpp
@@ -17,6 +17,21 @@ FOnlineAccountIdRegistryEOSGS::FOnlineAccountIdRegistryEOSGS()
 
 }
 
+FOnlineAccountIdRegistryEOSGS::FOnlineAccountIdRegistryEOSGS(EOnlineServices Services)
+	: Registry(Services)
+{
+
+}
+
+FAccountId FOnlineAccountIdRegistryEOSGS::FindOrAddAccountId(const EOS_ProductUserId ProductUserId)
+{
+	if (ensure(EOS_ProductUserId_IsValid(ProductUserId)))
+	{
+		return Registry.FindOrAddHandle(ProductUserId);
+	}
+	return Registry.GetInvalidHandle();
+}
+
 FAccountId FOnlineAccountIdRegistryEOSGS::FindAccountId(const EOS_ProductUserId ProductUserId) const
 {
 	return Registry.FindHandle(ProductUserId);
@@ -83,29 +98,22 @@ FAccountId FOnlineAccountIdRegistryEOSGS::FromReplicationData(const TArray<uint8
 
 IOnlineAccountIdRegistryEOSGS& FOnlineAccountIdRegistryEOSGS::GetRegistered()
 {
-	IOnlineAccountIdRegistry* Registry = FOnlineIdRegistryRegistry::Get().GetAccountIdRegistry(EOnlineServices::Epic);
-	check(Registry);
-	return *static_cast<IOnlineAccountIdRegistryEOSGS*>(Registry);
+	return FOnlineAccountIdRegistryEOSGS::GetRegistered(EOnlineServices::Epic);
 }
 
-FOnlineAccountIdRegistryEOSGS& FOnlineAccountIdRegistryEOSGS::Get()
+IOnlineAccountIdRegistryEOSGS& FOnlineAccountIdRegistryEOSGS::GetRegistered(EOnlineServices Services)
 {
-	static FOnlineAccountIdRegistryEOSGS Instance;
-	return Instance;
-}
+	check(Services == EOnlineServices::Epic || Services == EOnlineServices::EpicGame);
 
-FAccountId FOnlineAccountIdRegistryEOSGS::FindOrAddAccountId(const EOS_ProductUserId ProductUserId)
-{
-	if (ensure(EOS_ProductUserId_IsValid(ProductUserId)))
-	{
-		return Registry.FindOrAddHandle(ProductUserId);
-	}
-	return Registry.GetInvalidHandle();
+	IOnlineAccountIdRegistry* Registry = FOnlineIdRegistryRegistry::Get().GetAccountIdRegistry(Services);
+	check(Registry);
+	
+	return *static_cast<IOnlineAccountIdRegistryEOSGS*>(Registry);
 }
 
 EOS_ProductUserId GetProductUserId(const FAccountId& AccountId)
 {
-	return FOnlineAccountIdRegistryEOSGS::GetRegistered().GetProductUserId(AccountId);
+	return FOnlineAccountIdRegistryEOSGS::GetRegistered(AccountId.GetOnlineServicesType()).GetProductUserId(AccountId);
 }
 
 EOS_ProductUserId GetProductUserIdChecked(const FAccountId& AccountId)
@@ -117,12 +125,22 @@ EOS_ProductUserId GetProductUserIdChecked(const FAccountId& AccountId)
 
 FAccountId FindAccountId(const EOS_ProductUserId ProductUserId)
 {
-	return FOnlineAccountIdRegistryEOSGS::GetRegistered().FindAccountId(ProductUserId);
+	return FindAccountId(EOnlineServices::Epic, ProductUserId);
+}
+
+FAccountId FindAccountId(EOnlineServices Services, const EOS_ProductUserId ProductUserId)
+{
+	return FOnlineAccountIdRegistryEOSGS::GetRegistered(Services).FindAccountId(ProductUserId);
 }
 
 FAccountId FindAccountIdChecked(const EOS_ProductUserId ProductUserId)
 {
-	FAccountId Result = FindAccountId(ProductUserId);
+	return FindAccountIdChecked(EOnlineServices::Epic, ProductUserId);
+}
+
+FAccountId FindAccountIdChecked(EOnlineServices Services, const EOS_ProductUserId ProductUserId)
+{
+	FAccountId Result = FindAccountId(Services, ProductUserId);
 	check(Result.IsValid());
 	return Result;
 }
diff --git a/Engine/Plugins/Online/OnlineServicesEOSGS/Source/Private/Online/OnlineServicesEOSGS.cpp b/Engine/Plugins/Online/OnlineServicesEOSGS/Source/Private/Online/OnlineServicesEOSGS.cpp
--- a/Engine/Plugins/Online/OnlineServicesEOSGS/Source/Private/Online/OnlineServicesEOSGS.cpp
+++ b/Engine/Plugins/Online/OnlineServicesEOSGS/Source/Private/Online/OnlineServicesEOSGS.cpp
@@ -30,11 +30,8 @@ FOnlineServicesEOSGS::FOnlineServicesEOSGS(FName InInstanceName, FName InInstanc
 
 bool FOnlineServicesEOSGS::PreInit()
 {
-	FOnlineServicesEOSGSPlatformFactory& PlatformFactory = FOnlineServicesEOSGSPlatformFactory::Get();
-	EOSPlatformHandle = PlatformFactory.CreatePlatform(InstanceName, InstanceConfigName);
-	if (!EOSPlatformHandle)
+	if(!Super::PreInit())
 	{
-		UE_LOG(LogOnlineServices, Warning, TEXT("[%hs] InstanceName=%s InstanceConfigName=%s EOSPlatformHandle=nullptr."), __FUNCTION__, *InstanceName.ToString(), *InstanceConfigName.ToString());
 		return false;
 	}
 	
@@ -56,6 +53,7 @@ bool FOnlineServicesEOSGS::PreInit()
 void FOnlineServicesEOSGS::Destroy()
 {
 	Super::Destroy();
+
 #if WITH_ENGINE
 	if (SocketSubsystem)
 	{
diff --git a/Engine/Plugins/Online/OnlineServicesEOSGS/Source/Private/Online/OnlineServicesEOSGSModule.cpp b/Engine/Plugins/Online/OnlineServicesEOSGS/Source/Private/Online/OnlineServicesEOSGSModule.cpp
--- a/Engine/Plugins/Online/OnlineServicesEOSGS/Source/Private/Online/OnlineServicesEOSGSModule.cpp
+++ b/Engine/Plugins/Online/OnlineServicesEOSGS/Source/Private/Online/OnlineServicesEOSGSModule.cpp
@@ -2,6 +2,7 @@
 
 #include "Online/OnlineServicesEOSGSModule.h"
 
+#include "Misc/CommandLine.h"
 #include "Modules/ModuleManager.h"
 #include "Online/OnlineIdEOSGS.h"
 #include "Online/OnlineServicesEOSGS.h"
@@ -34,11 +35,23 @@ int FOnlineServicesEOSGSModule::GetRegistryPriority()
 
 void FOnlineServicesEOSGSModule::StartupModule()
 {
+	const bool bEnableEOSSplitPlugins = FParse::Param(FCommandLine::Get(), TEXT("EnableEOSSplitPlugins"));
+	if (bEnableEOSSplitPlugins)
+	{
+		FModuleManager::Get().LoadModuleChecked(TEXT("OnlineServicesEpicGameInternal"));
+		FModuleManager::Get().LoadModuleChecked(TEXT("OnlineServicesEpicAccountInternal"));
+
+		UE_LOG(LogOnlineServices, Verbose, TEXT("EnableEOSSplitPlugins launch parameter detected. Loading OnlineServicesEpicGameInternal and OnlineServicesEpicAccountInternal modules. Skipping FOnlineServicesEOSGSModule startup."));
+		return;
+	}
+
 	FModuleManager::Get().LoadModuleChecked(TEXT("OnlineServicesInterface"));
 	FModuleManager::Get().LoadModuleChecked(TEXT("EOSShared"));
 
 	FOnlineServicesRegistry::Get().RegisterServicesFactory(EOnlineServices::Epic, MakeUnique<FOnlineServicesFactoryEOSGS>(), GetRegistryPriority());
-	FOnlineIdRegistryRegistry::Get().RegisterAccountIdRegistry(EOnlineServices::Epic, &FOnlineAccountIdRegistryEOSGS::Get(), GetRegistryPriority());
+
+	static FOnlineAccountIdRegistryEOSGS AccountIdRegistry(EOnlineServices::Epic);
+	FOnlineIdRegistryRegistry::Get().RegisterAccountIdRegistry(EOnlineServices::Epic, &AccountIdRegistry, GetRegistryPriority());
 	FOnlineIdRegistryRegistry::Get().RegisterSessionIdRegistry(EOnlineServices::Epic, &FOnlineSessionIdRegistryEOSGS::Get(), GetRegistryPriority());
 	FOnlineIdRegistryRegistry::Get().RegisterSessionInviteIdRegistry(EOnlineServices::Epic, &FOnlineSessionInviteIdRegistryEOSGS::Get(), GetRegistryPriority());
 
diff --git a/Engine/Plugins/Online/OnlineServicesEOSGS/Source/Private/Online/SessionsEOSGS.cpp b/Engine/Plugins/Online/OnlineServicesEOSGS/Source/Private/Online/SessionsEOSGS.cpp
--- a/Engine/Plugins/Online/OnlineServicesEOSGS/Source/Private/Online/SessionsEOSGS.cpp
+++ b/Engine/Plugins/Online/OnlineServicesEOSGS/Source/Private/Online/SessionsEOSGS.cpp
@@ -7,7 +7,7 @@
 #include "Online/AuthEOSGS.h"
 #include "Online/LobbiesCommonTypes.h"
 #include "Online/NboSerializerEOSGSSvc.h"
-#include "Online/OnlineErrorEOSGS.h"
+#include "Online/OnlineErrorEpicCommon.h"
 #include "Online/OnlineServicesEOSGS.h"
 #include "Online/SessionsEOSGSTypes.h"
 
@@ -1623,4 +1623,9 @@ void FSessionsEOSGS::ReadSessionFromPacket(FNboSerializeFromBuffer& Packet, FSes
 	SerializeFromBuffer(Packet, Session.SessionMembers);
 }
 
+FAccountId FSessionsEOSGS::FindAccountId(const EOS_ProductUserId ProductUserId)
+{
+	return UE::Online::FindAccountId(Services.GetServicesProvider(), ProductUserId);
+}
+
 /* UE::Online */ }
diff --git a/Engine/Plugins/Online/OnlineServicesEOSGS/Source/Private/Online/StatsEOSGS.cpp b/Engine/Plugins/Online/OnlineServicesEOSGS/Source/Private/Online/StatsEOSGS.cpp
--- a/Engine/Plugins/Online/OnlineServicesEOSGS/Source/Private/Online/StatsEOSGS.cpp
+++ b/Engine/Plugins/Online/OnlineServicesEOSGS/Source/Private/Online/StatsEOSGS.cpp
@@ -5,7 +5,7 @@
 #include "EOSShared.h"
 #include "IEOSSDKManager.h"
 #include "Online/AuthEOSGS.h"
-#include "Online/OnlineErrorEOSGS.h"
+#include "Online/OnlineErrorEpicCommon.h"
 #include "Online/OnlineIdEOSGS.h"
 #include "Online/OnlineServicesEOSGS.h"
 
@@ -144,7 +144,7 @@ TOnlineAsyncOpHandle<FUpdateStats> FStatsEOSGS::UpdateStats(FUpdateStats::Params
 namespace Private
 {
 
-void QueryStatsEOS(EOS_HStats StatsHandle, const FAccountId& LocalAccountId, const FAccountId& TargetAccountId, const TArray<FString>& StatNames, TPromise<const EOS_Stats_OnQueryStatsCompleteCallbackInfo*>&& Promise)
+void QueryStatsEOS(EOnlineServices Services, EOS_HStats StatsHandle, const FAccountId& LocalAccountId, const FAccountId& TargetAccountId, const TArray<FString>& StatNames, TPromise<const EOS_Stats_OnQueryStatsCompleteCallbackInfo*>&& Promise)
 {
 	EOS_Stats_QueryStatsOptions Options;
 	Options.ApiVersion = 3;
@@ -239,7 +239,7 @@ TOnlineAsyncOpHandle<FQueryStats> FStatsEOSGS::QueryStats(FQueryStats::Params&&
 		.Then([this](TOnlineAsyncOp<FQueryStats>& InAsyncOp, TPromise<const EOS_Stats_OnQueryStatsCompleteCallbackInfo*>&& Promise)
 		{
 			const FQueryStats::Params& Params = InAsyncOp.GetParams();
-			Private::QueryStatsEOS(StatsHandle, Params.LocalAccountId, Params.TargetAccountId, Params.StatNames, MoveTemp(Promise));
+			Private::QueryStatsEOS(Services.GetServicesProvider(), StatsHandle, Params.LocalAccountId, Params.TargetAccountId, Params.StatNames, MoveTemp(Promise));
 		})
 		.Then([this](TOnlineAsyncOp<FQueryStats>& InAsyncOp, const EOS_Stats_OnQueryStatsCompleteCallbackInfo* Data)
 		{
@@ -289,7 +289,7 @@ TOnlineAsyncOpHandle<FBatchQueryStats> FStatsEOSGS::BatchQueryStats(FBatchQueryS
 		{
 			Op->Then([this, TargetAccountId](TOnlineAsyncOp<FBatchQueryStats>& InAsyncOp, TPromise<const EOS_Stats_OnQueryStatsCompleteCallbackInfo*>&& Promise)
 			{
-				Private::QueryStatsEOS(StatsHandle, InAsyncOp.GetParams().LocalAccountId, TargetAccountId, InAsyncOp.GetParams().StatNames, MoveTemp(Promise));
+				Private::QueryStatsEOS(Services.GetServicesProvider(), StatsHandle, InAsyncOp.GetParams().LocalAccountId, TargetAccountId, InAsyncOp.GetParams().StatNames, MoveTemp(Promise));
 			})
 			.Then([this, TargetAccountId](TOnlineAsyncOp<FBatchQueryStats>& InAsyncOp, const EOS_Stats_OnQueryStatsCompleteCallbackInfo* Data)
 			{
diff --git a/Engine/Plugins/Online/OnlineServicesEOSGS/Source/Private/Online/TitleFileEOSGS.cpp b/Engine/Plugins/Online/OnlineServicesEOSGS/Source/Private/Online/TitleFileEOSGS.cpp
--- a/Engine/Plugins/Online/OnlineServicesEOSGS/Source/Private/Online/TitleFileEOSGS.cpp
+++ b/Engine/Plugins/Online/OnlineServicesEOSGS/Source/Private/Online/TitleFileEOSGS.cpp
@@ -5,7 +5,7 @@
 #include "EOSShared.h"
 #include "IEOSSDKManager.h"
 #include "Online/AuthEOSGS.h"
-#include "Online/OnlineErrorEOSGS.h"
+#include "Online/OnlineErrorEpicCommon.h"
 #include "Online/OnlineIdEOSGS.h"
 #include "Online/OnlineServicesEOSGS.h"
 
diff --git a/Engine/Plugins/Online/OnlineServicesEOSGS/Source/Private/Online/UserFileEOSGS.cpp b/Engine/Plugins/Online/OnlineServicesEOSGS/Source/Private/Online/UserFileEOSGS.cpp
--- a/Engine/Plugins/Online/OnlineServicesEOSGS/Source/Private/Online/UserFileEOSGS.cpp
+++ b/Engine/Plugins/Online/OnlineServicesEOSGS/Source/Private/Online/UserFileEOSGS.cpp
@@ -5,7 +5,7 @@
 #include "EOSShared.h"
 #include "IEOSSDKManager.h"
 #include "Online/AuthEOSGS.h"
-#include "Online/OnlineErrorEOSGS.h"
+#include "Online/OnlineErrorEpicCommon.h"
 #include "Online/OnlineIdEOSGS.h"
 #include "Online/OnlineServicesEOSGS.h"
 
diff --git a/Engine/Plugins/Online/OnlineServicesEOSGS/Source/Public/Online/AchievementsEOSGS.h b/Engine/Plugins/Online/OnlineServicesEOSGS/Source/Public/Online/AchievementsEOSGS.h
--- a/Engine/Plugins/Online/OnlineServicesEOSGS/Source/Public/Online/AchievementsEOSGS.h
+++ b/Engine/Plugins/Online/OnlineServicesEOSGS/Source/Public/Online/AchievementsEOSGS.h
@@ -38,6 +38,8 @@ public:
 protected:
 	void HandleAchievementsUnlocked(const EOS_Achievements_OnAchievementsUnlockedCallbackV2Info* Data);
 
+	FAccountId FindAccountId(const EOS_ProductUserId ProductUserId);
+
 	EOS_HAchievements AchievementsHandle = nullptr;
 
 	using FAchievementDefinitionMap = TMap<FString, FAchievementDefinition>;
diff --git a/Engine/Plugins/Online/OnlineServicesEOSGS/Source/Public/Online/AuthEOSGS.h b/Engine/Plugins/Online/OnlineServicesEOSGS/Source/Public/Online/AuthEOSGS.h
--- a/Engine/Plugins/Online/OnlineServicesEOSGS/Source/Public/Online/AuthEOSGS.h
+++ b/Engine/Plugins/Online/OnlineServicesEOSGS/Source/Public/Online/AuthEOSGS.h
@@ -191,7 +191,7 @@ class ONLINESERVICESEOSGS_API FAuthEOSGS : public FAuthCommon
 public:
 	using Super = FAuthCommon;
 
-	FAuthEOSGS(FOnlineServicesEOSGS& InOwningSubsystem);
+	FAuthEOSGS(FOnlineServicesCommon& InOwningSubsystem);
 	virtual ~FAuthEOSGS() = default;
 
 	// Begin IOnlineComponent
diff --git a/Engine/Plugins/Online/OnlineServicesEOSGS/Source/Public/Online/LeaderboardsEOSGS.h b/Engine/Plugins/Online/OnlineServicesEOSGS/Source/Public/Online/LeaderboardsEOSGS.h
--- a/Engine/Plugins/Online/OnlineServicesEOSGS/Source/Public/Online/LeaderboardsEOSGS.h
+++ b/Engine/Plugins/Online/OnlineServicesEOSGS/Source/Public/Online/LeaderboardsEOSGS.h
@@ -30,6 +30,8 @@ public:
 	virtual TOnlineAsyncOpHandle<FReadEntriesAroundUser> ReadEntriesAroundUser(FReadEntriesAroundUser::Params&& Params) override;
 
 private:
+	void ReadEntriesInRange(uint32 StartIndex, uint32 EndIndex, TArray<FLeaderboardEntry>& OutEntries);
+
 	EOS_HLeaderboards LeaderboardsHandle = nullptr;
 };
 
diff --git a/Engine/Plugins/Online/OnlineServicesEOSGS/Source/Public/Online/LobbiesEOSGS.h b/Engine/Plugins/Online/OnlineServicesEOSGS/Source/Public/Online/LobbiesEOSGS.h
--- a/Engine/Plugins/Online/OnlineServicesEOSGS/Source/Public/Online/LobbiesEOSGS.h
+++ b/Engine/Plugins/Online/OnlineServicesEOSGS/Source/Public/Online/LobbiesEOSGS.h
@@ -215,6 +215,8 @@ private:
 	void HandleLobbyInviteAccepted(const EOS_Lobby_LobbyInviteAcceptedCallbackInfo* Data);
 	void HandleJoinLobbyAccepted(const EOS_Lobby_JoinLobbyAcceptedCallbackInfo* Data);
 
+	FAccountId FindAccountId(const EOS_ProductUserId ProductUserId);
+
 protected:
 #if !UE_BUILD_SHIPPING
 	static void CheckMetadata();
diff --git a/Engine/Plugins/Online/OnlineServicesEOSGS/Source/Public/Online/NboSerializerEOSGSSvc.h b/Engine/Plugins/Online/OnlineServicesEOSGS/Source/Public/Online/NboSerializerEOSGSSvc.h
--- a/Engine/Plugins/Online/OnlineServicesEOSGS/Source/Public/Online/NboSerializerEOSGSSvc.h
+++ b/Engine/Plugins/Online/OnlineServicesEOSGS/Source/Public/Online/NboSerializerEOSGSSvc.h
@@ -22,7 +22,7 @@ namespace NboSerializerEOSGSSvc {
 
 inline void SerializeToBuffer(FNboSerializeToBuffer& Ar, const FAccountId& UniqueId)
 {
-	TArray<uint8> Data = FOnlineAccountIdRegistryEOSGS::GetRegistered().ToReplicationData(UniqueId);
+	TArray<uint8> Data = FOnlineAccountIdRegistryEOSGS::GetRegistered(UniqueId.GetOnlineServicesType()).ToReplicationData(UniqueId);
 	Ar << Data.Num();
 	Ar.WriteBinary(Data.GetData(), Data.Num());
 }
@@ -52,7 +52,7 @@ inline void SerializeFromBuffer(FNboSerializeFromBuffer& Ar, FAccountId& UniqueI
 	int32 Size;
 	Ar >> Size;
 	Ar.ReadBinaryArray(Data, Size);
-	UniqueId = FOnlineAccountIdRegistryEOSGS::GetRegistered().FromReplicationData(Data);
+	UniqueId = FOnlineAccountIdRegistryEOSGS::GetRegistered(UniqueId.GetOnlineServicesType()).FromReplicationData(Data);
 }
 
 inline void SerializeFromBuffer(FNboSerializeFromBuffer& Ar, FOnlineSessionId& SessionId)
diff --git a/Engine/Plugins/Online/OnlineServicesEOSGS/Source/Public/Online/OnlineErrorEOSGS.h b/Engine/Plugins/Online/OnlineServicesEOSGS/Source/Public/Online/OnlineErrorEOSGS.h
--- a/Engine/Plugins/Online/OnlineServicesEOSGS/Source/Public/Online/OnlineErrorEOSGS.h
+++ b/Engine/Plugins/Online/OnlineServicesEOSGS/Source/Public/Online/OnlineErrorEOSGS.h
@@ -2,49 +2,8 @@
 
 #pragma once
 
-#include "Online/OnlineError.h"
+#include "Misc/CoreMiscDefines.h"
 
-#if defined(EOS_PLATFORM_BASE_FILE_NAME)
-#include EOS_PLATFORM_BASE_FILE_NAME
-#endif
-#include "eos_common.h"
+UE_DEPRECATED_HEADER(5.6, "Use OnlineErrorEpicCommon.h instead of OnlineErrorEOSGS.h.")
 
-/*
- *	Proper usage:
- *		Errors::FromEOSResult(EOSResult::EOS_PlayerDataStorage_FileSizeTooLarge);
- *	Certain EOS errors are predefined to have a common error parent type (see Internal_EOSWrapInner). i.e.
- *		(Errors::FromEOSResult(EOSResult::EOS_NoConnection) == Errors::NoConnection()) == true
- */
-
-namespace UE::Online::Errors {
-
-	UE_ONLINE_ERROR_CATEGORY(EOS, ThirdPartyPlugin, 0x4, "EOS")
-
-	typedef TFunction<FOnlineError(FOnlineError&& Error, EOS_EResult Result)> FErrorMapperEosFn;
-
-	ONLINESERVICESEOSGS_API FOnlineError MapCommonEOSError(FOnlineError&& Error, EOS_EResult Result);
-	ONLINESERVICESEOSGS_API ErrorCodeType ErrorCodeFromEOSResult(EOS_EResult Result);
-	ONLINESERVICESEOSGS_API FOnlineError FromEOSResult(EOS_EResult Result, FErrorMapperEosFn&& MapperFn = &MapCommonEOSError);
-
-/* UE::Online::Errors */ }
-
-// These are extern'd out of the namespace for some of the Catch API that can't see these operators if they are inside
-inline bool operator==(const UE::Online::FOnlineError& OnlineError, EOS_EResult EosResult)
-{
-	return OnlineError == UE::Online::Errors::ErrorCodeFromEOSResult(EosResult);
-}
-
-inline bool operator==(EOS_EResult EosResult, const UE::Online::FOnlineError& OnlineError)
-{
-	return OnlineError == EosResult;
-}
-
-inline bool operator!=(const UE::Online::FOnlineError& OnlineError, EOS_EResult EosResult)
-{
-	return !(OnlineError == EosResult);
-}
-
-inline bool operator!=(EOS_EResult EosResult, const UE::Online::FOnlineError& OnlineError)
-{
-	return OnlineError != EosResult;
-}
+#include "Online/OnlineErrorEpicCommon.h"
\ No newline at end of file
diff --git a/Engine/Plugins/Online/OnlineServicesEOSGS/Source/Public/Online/OnlineIdEOSGS.h b/Engine/Plugins/Online/OnlineServicesEOSGS/Source/Public/Online/OnlineIdEOSGS.h
--- a/Engine/Plugins/Online/OnlineServicesEOSGS/Source/Public/Online/OnlineIdEOSGS.h
+++ b/Engine/Plugins/Online/OnlineServicesEOSGS/Source/Public/Online/OnlineIdEOSGS.h
@@ -15,6 +15,7 @@ namespace UE::Online {
 class IOnlineAccountIdRegistryEOSGS : public IOnlineAccountIdRegistry
 {
 public:
+	virtual FAccountId FindOrAddAccountId(const EOS_ProductUserId ProductUserId) = 0;
 	virtual FAccountId FindAccountId(EOS_ProductUserId ProductUserId) const = 0;
 	virtual EOS_ProductUserId GetProductUserId(const FAccountId& AccountId) const = 0;
 };
@@ -22,14 +23,18 @@ public:
 /**
  * Account id registry specifically for EOS id's which are segmented.
  */
-class FOnlineAccountIdRegistryEOSGS
+class ONLINESERVICESEOSGS_API FOnlineAccountIdRegistryEOSGS
 	: public IOnlineAccountIdRegistryEOSGS
 {
 public:
+	UE_DEPRECATED(5.6, "This method is deprecated, please use the new version taking a EOnlineServices parameter")
 	FOnlineAccountIdRegistryEOSGS();
+	FOnlineAccountIdRegistryEOSGS(EOnlineServices Services);
+
 	virtual ~FOnlineAccountIdRegistryEOSGS() = default;
 
 	// Begin IOnlineAccountIdRegistryEOSGS
+	virtual FAccountId FindOrAddAccountId(const EOS_ProductUserId ProductUserId) override;
 	virtual FAccountId FindAccountId(const EOS_ProductUserId ProductUserId) const override;
 	virtual EOS_ProductUserId GetProductUserId(const FAccountId& AccountId) const override;
 	// End IOnlineAccountIdRegistryEOSGS
@@ -41,27 +46,31 @@ public:
 	virtual FAccountId FromReplicationData(const TArray<uint8>& ReplicationData) override;
 	// End IOnlineAccountIdRegistry
 
+	UE_DEPRECATED(5.6, "This method is deprecated, please use the new version taking a EOnlineServices parameter")
 	static IOnlineAccountIdRegistryEOSGS& GetRegistered();
 
-private:
-	// FAuthEOSGS is the only thing that should be able to create PUID-only net ids in this registry, in its resolve methods.
-	friend class FAuthEOSGS;
-	friend class FOnlineServicesEOSGSModule;
-	static FOnlineAccountIdRegistryEOSGS& Get();
-	FAccountId FindOrAddAccountId(const EOS_ProductUserId ProductUserId);
+	static IOnlineAccountIdRegistryEOSGS& GetRegistered(EOnlineServices Services);
 
+private:
 	TOnlineBasicAccountIdRegistry<EOS_ProductUserId> Registry;
 };
 
 EOS_ProductUserId ONLINESERVICESEOSGS_API GetProductUserId(const FAccountId& AccountId);
 EOS_ProductUserId ONLINESERVICESEOSGS_API GetProductUserIdChecked(const FAccountId& AccountId);
-FAccountId ONLINESERVICESEOSGS_API FindAccountId(const EOS_ProductUserId EpicAccountId);
-FAccountId ONLINESERVICESEOSGS_API FindAccountIdChecked(const EOS_ProductUserId EpicAccountId);
+
+UE_DEPRECATED(5.6, "This method is deprecated, please use the new version taking a EOnlineServices parameter")
+FAccountId ONLINESERVICESEOSGS_API FindAccountId(const EOS_ProductUserId ProductUserId);
+FAccountId ONLINESERVICESEOSGS_API FindAccountId(EOnlineServices Services, const EOS_ProductUserId ProductUserId);
+
+UE_DEPRECATED(5.6, "This method is deprecated, please use the new version taking a EOnlineServices parameter")
+FAccountId ONLINESERVICESEOSGS_API FindAccountIdChecked(const EOS_ProductUserId ProductUserId);
+FAccountId ONLINESERVICESEOSGS_API FindAccountIdChecked(EOnlineServices Services, const EOS_ProductUserId ProductUserId);
 
 template<typename IdType>
+UE_DEPRECATED(5.6, "This method is deprecated, please add an equivalent method or perform the relevant checks in your calling code")
 inline bool ValidateOnlineId(const TOnlineId<IdType> OnlineId)
 {
-	return OnlineId.GetOnlineServicesType() == EOnlineServices::Epic && OnlineId.IsValid();
+	return OnlineId.IsValid() && (OnlineId.GetOnlineServicesType() == EOnlineServices::Epic);
 }
 
 } /* namespace UE::Online */
diff --git a/Engine/Plugins/Online/OnlineServicesEOSGS/Source/Public/Online/OnlineServicesEOSGS.h b/Engine/Plugins/Online/OnlineServicesEOSGS/Source/Public/Online/OnlineServicesEOSGS.h
--- a/Engine/Plugins/Online/OnlineServicesEOSGS/Source/Public/Online/OnlineServicesEOSGS.h
+++ b/Engine/Plugins/Online/OnlineServicesEOSGS/Source/Public/Online/OnlineServicesEOSGS.h
@@ -2,7 +2,7 @@
 
 #pragma once
 
-#include "Online/OnlineServicesCommon.h"
+#include "Online/OnlineServicesEpicCommon.h"
 
 #if defined(EOS_PLATFORM_BASE_FILE_NAME)
 #include EOS_PLATFORM_BASE_FILE_NAME
@@ -17,14 +17,15 @@ using IEOSPlatformHandlePtr = TSharedPtr<class IEOSPlatformHandle>;
 
 namespace UE::Online {
 
-class ONLINESERVICESEOSGS_API FOnlineServicesEOSGS : public FOnlineServicesCommon
+class ONLINESERVICESEOSGS_API FOnlineServicesEOSGS : public FOnlineServicesEpicCommon
 {
 public:
-	using Super = FOnlineServicesCommon;
+	using Super = FOnlineServicesEpicCommon;
 
 	FOnlineServicesEOSGS(FName InstanceName, FName InstanceConfigName);
 	virtual ~FOnlineServicesEOSGS() = default;
 
+	virtual bool PreInit() override;
 	virtual void Destroy() override;
 	virtual void RegisterComponents() override;
 	virtual TOnlineResult<FGetResolvedConnectString> GetResolvedConnectString(FGetResolvedConnectString::Params&& Params) override;
@@ -36,7 +37,7 @@ public:
 	static const TCHAR* GetConfigNameStatic() { return GetServiceConfigNameStatic(); }
 	static const TCHAR* GetServiceConfigNameStatic() { return TEXT("EOS"); }
 
-	bool PreInit();
+
 
 protected:
 	IEOSPlatformHandlePtr EOSPlatformHandle;
diff --git a/Engine/Plugins/Online/OnlineServicesEOSGS/Source/Public/Online/OnlineServicesEOSGSTypes.h b/Engine/Plugins/Online/OnlineServicesEOSGS/Source/Public/Online/OnlineServicesEOSGSTypes.h
--- a/Engine/Plugins/Online/OnlineServicesEOSGS/Source/Public/Online/OnlineServicesEOSGSTypes.h
+++ b/Engine/Plugins/Online/OnlineServicesEOSGS/Source/Public/Online/OnlineServicesEOSGSTypes.h
@@ -2,270 +2,6 @@
 
 #pragma once
 
-#include "Online/OnlineAsyncOp.h"
-#include "Online/OnlineErrorEOSGS.h"
-#include "Async/Future.h"
-
-#include "EOSShared.h"
-
-namespace UE::Online {
-
-namespace Private {
+#include "EOSSharedTypes.h"
 
 #define LOBBIESEOSGS_LOBBY_ID_OVERRIDE_ATTRIBUTE_KEY TEXT("LOBBY_ID_OVERRIDE")
-
-/** Class to handle all callbacks generically using a future to forward callback results */
-template<typename CallbackType>
-class TEOSCallback
-{
-	using CallbackFuncType = void (EOS_CALL*)(const CallbackType*);
-public:
-	TEOSCallback() = default;
-	TEOSCallback(TPromise<const CallbackType*>&& InPromise) : Promise(MoveTemp(InPromise)) {}
-	virtual ~TEOSCallback() = default;
-
-	operator CallbackFuncType()
-	{
-		return &CallbackImpl;
-	}
-
-	TFuture<const CallbackType*> GetFuture() { return Promise.GetFuture(); }
-
-private:
-	TPromise<const CallbackType*> Promise;
-
-	static void EOS_CALL CallbackImpl(const CallbackType* Data)
-	{
-		if (EOS_EResult_IsOperationComplete(Data->ResultCode) == EOS_FALSE)
-		{
-			// Ignore
-			return;
-		}
-		check(IsInGameThread());
-
-		TEOSCallback* CallbackThis = reinterpret_cast<TEOSCallback*>(Data->ClientData);
-		check(CallbackThis);
-
-		CallbackThis->Promise.EmplaceValue(Data);
-
-		delete CallbackThis;
-	}
-};
-
-template<typename Function> struct FEOSAsyncCallbackTraits;
-template<typename TCallbackData>
-struct FEOSAsyncCallbackTraits<void(*)(const TCallbackData*)>
-{
-	typedef TCallbackData CallbackDataType;
-};
-
-template<typename Function> struct FEOSAsyncTraits;
-template<typename TEOSHandle, typename TEOSParameters, typename TEOSCallback>
-struct FEOSAsyncTraits<void(*)(TEOSHandle, const TEOSParameters*, void*, const TEOSCallback)>
-{
-	typedef TEOSCallback CallbackType;
-	typedef typename FEOSAsyncCallbackTraits<TEOSCallback>::CallbackDataType CallbackDataType;
-};
-
-/* Private */ }
-
-/**
- * This method requires a Promise that will be fulfilled by the EOS callback.
- * This is primarily intended to be used with continuations that take a TPromise. e.g.:
- *
- *		InAsyncOp.Then([this, LoginOptions, Credentials](TOnlineAsyncOp<FAuthLogin>& InAsyncOp, TPromise<const EOS_Auth_LoginCallbackInfo*>&& Promise) mutable
- *		{
- *			LoginOptions.Credentials = &Credentials;
- *			EOS_Async(EOS_Auth_Login, AuthHandle, LoginOptions, MoveTemp(Promise));
- *		})
- * 
- * Note that when a continuation takes a Promise parameter, the Future for that Promise is already bound
- * to the next continuation in the Op chain. Therefore, as soon as the EOS_Auth_Login callback is called,
- * the promise will be fulfilled, and the next continuation will run immediately. This ensures that even
- * if the SDK calls the callback immediately (i.e. before EOS_Async returns), the continuation is already
- * bound, and we can safely consume the CallbackInfo. Note that the CallbackInfo is only valid for the
- * duration of the callback, so it is not safe to bind a continuation _after_ calling the EOS method.
- */
-// Warning: This signature will crash if the async method returns immediately.
-// The callback object is deleted inside the callback before EOS_Async returns the future back to the user.
-template<typename TEOSResult, typename TEOSHandle, typename TEOSParameters, typename TEOSFn>
-decltype(auto) EOS_Async(TEOSFn EOSFn, TEOSHandle EOSHandle, TEOSParameters Parameters, TPromise<const TEOSResult*>&& Promise)
-{
-	Private::TEOSCallback<TEOSResult>* Callback = new Private::TEOSCallback<TEOSResult>(Forward< TPromise<const TEOSResult*>>(Promise));
-	return EOSFn(EOSHandle, &Parameters, Callback, *Callback);
-}
-
-/**
- * This method requires a Callback (lambda / function pointer) which will be called by the EOS callback
- * with the CallbackInfo.
- */
-template<typename TEOSHandle, typename TEOSParameters, typename TEOSFn, typename TCallback>
-decltype(auto) EOS_Async(TEOSFn EOSFn, TEOSHandle EOSHandle, const TEOSParameters& Parameters, TCallback&& Callback)
-{
-	using CallbackDataType = typename Private::FEOSAsyncTraits<TEOSFn>::CallbackDataType;
-	TUniquePtr<TCallback> CallbackObj = MakeUnique<TCallback>(Forward<TCallback>(Callback));
-	return EOSFn(EOSHandle, &Parameters, CallbackObj.Release(),
-	static_cast<void(EOS_CALL*)(const CallbackDataType* Data)>([](const CallbackDataType* Data)
-	{
-		if (EOS_EResult_IsOperationComplete(Data->ResultCode))
-		{
-			check(IsInGameThread());
-			TUniquePtr<TCallback> CallbackObj(reinterpret_cast<TCallback*>(Data->ClientData));
-			check(CallbackObj);
-			(*CallbackObj)(Data);
-		}
-	}));
-}
-
-class FEOSEventRegistration
-{
-public:
-	virtual ~FEOSEventRegistration() = default;
-};
-typedef TUniquePtr<FEOSEventRegistration> FEOSEventRegistrationPtr;
-
-using EOSEventRegistration UE_DEPRECATED(5.6, "EOSEventRegistration renamed to FEOSEventRegistration") = FEOSEventRegistration;
-using EOSEventRegistrationPtr UE_DEPRECATED(5.6, "EOSEventRegistrationPtr renamed to FEOSEventRegistrationPtr") = FEOSEventRegistrationPtr;
-
-/** 
-* EOS event registration utility for binding an EOS notifier registration to a RAII object which handles
-* unregistering when it exits scope. Intended to be used from a TOnlineComponent class.
-* 
-* Example:
-*	EOSEventRegistrationPtr OnLobbyUpdatedEOSEventRegistration = EOS_RegisterComponentEventHandler(
-*		this,
-*		LobbyHandle,
-*		EOS_LOBBY_ADDNOTIFYLOBBYUPDATERECEIVED_API_LATEST,
-*		&EOS_Lobby_AddNotifyLobbyUpdateReceived,
-*		&EOS_Lobby_RemoveNotifyLobbyUpdateReceived,
-*		&FLobbiesEOS::HandleLobbyUpdated);
-**/
-template <
-	typename ComponentHandlerClass,
-	typename EOSHandle,
-	typename EOSNotfyRegisterFunction,
-	typename EOSNotfyUnregisterFunction,
-	typename ComponentHandlerFunction>
-FEOSEventRegistrationPtr EOS_RegisterComponentEventHandler(
-	ComponentHandlerClass* HandlerClass,
-	EOSHandle ClientHandle,
-	int32_t ApiVersion,
-	EOSNotfyRegisterFunction NotfyRegisterFunction,
-	EOSNotfyUnregisterFunction NotfyUnregisterFunction,
-	ComponentHandlerFunction HandlerFunction);
-
-namespace Private {
-
-template<typename Function> struct TEOSCallbackTraitsBase;
-
-template<typename Component, typename EventData>
-struct TEOSCallbackTraitsBase<void (Component::*)(const EventData*)>
-{
-	typedef Component ComponentType;
-	typedef EventData EventDataType;
-};
-
-template<typename Function>
-struct TEOSCallbackTraits : public TEOSCallbackTraitsBase<typename std::remove_reference<Function>::type>
-{
-};
-
-template<typename Function> struct TEOSNotifyRegisterTraits;
-
-template<typename EOSHandle, typename Options, typename ClientData, typename NotificationFn, typename NotificationId>
-struct TEOSNotifyRegisterTraits<NotificationId (*)(EOSHandle, const Options*, ClientData*, NotificationFn)>
-{
-	typedef Options OptionsType;
-	typedef NotificationId NotificationIdType;
-};
-
-template <
-	typename ComponentHandlerClass,
-	typename EOSHandle,
-	typename EOSNotfyRegisterFunction,
-	typename EOSNotfyUnregisterFunction,
-	typename ComponentHandlerFunction>
-class FEOSEventRegistrationImpl : public FEOSEventRegistration
-{
-public:
-	typedef FEOSEventRegistrationImpl<
-		ComponentHandlerClass,
-		EOSHandle,
-		EOSNotfyRegisterFunction,
-		EOSNotfyUnregisterFunction,
-		ComponentHandlerFunction> ThisClass;
-
-	FEOSEventRegistrationImpl(
-		ComponentHandlerClass* HandlerClass,
-		EOSHandle ClientHandle,
-		int32_t ApiVersion,
-		EOSNotfyRegisterFunction NotfyRegisterFunction,
-		EOSNotfyUnregisterFunction NotfyUnregisterFunction,
-		ComponentHandlerFunction HandlerFunction)
-		: HandlerClass(HandlerClass)
-		, ClientHandle(ClientHandle)
-		, NotfyUnregisterFunction(NotfyUnregisterFunction)
-		, HandlerFunction(HandlerFunction)
-	{
-		using EventDataType = typename TEOSCallbackTraits<ComponentHandlerFunction>::EventDataType;
-		typename TEOSNotifyRegisterTraits<EOSNotfyRegisterFunction>::OptionsType Options = { };
-		Options.ApiVersion = ApiVersion;
-		NotificationId = NotfyRegisterFunction(ClientHandle, &Options, this,
-		static_cast<void(EOS_CALL*)(const EventDataType* Data)>([](const EventDataType* Data)
-		{
-			ThisClass* This = reinterpret_cast<ThisClass*>(Data->ClientData);
-			(This->HandlerClass->*This->HandlerFunction)(Data);
-		}));
-	}
-
-	virtual ~FEOSEventRegistrationImpl()
-	{
-		NotfyUnregisterFunction(ClientHandle, NotificationId);
-	}
-
-private:
-	FEOSEventRegistrationImpl() = delete;
-	FEOSEventRegistrationImpl(const FEOSEventRegistrationImpl&) = delete;
-	FEOSEventRegistrationImpl& operator=(const FEOSEventRegistrationImpl&) = delete;
-
-	typedef typename TEOSNotifyRegisterTraits<EOSNotfyRegisterFunction>::NotificationIdType NotificationIdType;
-
-	NotificationIdType NotificationId;
-	ComponentHandlerClass* HandlerClass;
-	EOSHandle ClientHandle;
-	EOSNotfyUnregisterFunction NotfyUnregisterFunction;
-	ComponentHandlerFunction HandlerFunction;
-};
-
-/* Private */ }
-
-template <
-	typename ComponentHandlerClass,
-	typename EOSHandle,
-	typename EOSNotfyRegisterFunction,
-	typename EOSNotfyUnregisterFunction,
-	typename ComponentHandlerFunction>
-FEOSEventRegistrationPtr EOS_RegisterComponentEventHandler(
-	ComponentHandlerClass* HandlerClass,
-	EOSHandle ClientHandle,
-	int32_t ApiVersion,
-	EOSNotfyRegisterFunction NotfyRegisterFunction,
-	EOSNotfyUnregisterFunction NotfyUnregisterFunction,
-	ComponentHandlerFunction HandlerFunction)
-{
-	return MakeUnique<
-		Private::FEOSEventRegistrationImpl<
-			ComponentHandlerClass,
-			EOSHandle,
-			EOSNotfyRegisterFunction,
-			EOSNotfyUnregisterFunction,
-			ComponentHandlerFunction>>(
-			HandlerClass,
-			ClientHandle,
-			ApiVersion,
-			NotfyRegisterFunction,
-			NotfyUnregisterFunction,
-			HandlerFunction);
-}
-
-/* UE::Online */ }
diff --git a/Engine/Plugins/Online/OnlineServicesEOSGS/Source/Public/Online/SessionsEOSGS.h b/Engine/Plugins/Online/OnlineServicesEOSGS/Source/Public/Online/SessionsEOSGS.h
--- a/Engine/Plugins/Online/OnlineServicesEOSGS/Source/Public/Online/SessionsEOSGS.h
+++ b/Engine/Plugins/Online/OnlineServicesEOSGS/Source/Public/Online/SessionsEOSGS.h
@@ -279,6 +279,8 @@ private:
 	virtual void AppendSessionToPacket(FNboSerializeToBuffer& Packet, const FSessionLAN& Session) override;
 	virtual void ReadSessionFromPacket(FNboSerializeFromBuffer& Packet, FSessionLAN& Session) override;
 
+	FAccountId FindAccountId(const EOS_ProductUserId ProductUserId);
+
 protected:
 	EOS_HSessions SessionsHandle = nullptr;
 
diff --git a/Engine/Plugins/Online/OnlineServicesEpicCommon/OnlineServicesEpicCommon.uplugin b/Engine/Plugins/Online/OnlineServicesEpicCommon/OnlineServicesEpicCommon.uplugin
new file mode 100644
--- /dev/null
+++ b/Engine/Plugins/Online/OnlineServicesEpicCommon/OnlineServicesEpicCommon.uplugin
@@ -0,0 +1,38 @@
+{
+	"FileVersion": 3,
+
+	"FriendlyName": "Online Services Epic Common",
+	"Version": 1,
+	"VersionName": "1.0",
+	"Description": "Common Online Services implementation with shared funcionality between Epic Account Services and Epic Game Services.",
+	"Category": "Online Platform",
+	"CreatedBy": "Epic Games, Inc.",
+	"CreatedByURL": "https://epicgames.com",
+	"EnabledByDefault": false,
+
+	"Modules": [
+		{
+			"Name": "OnlineServicesEpicCommon",
+			"Type": "Runtime",
+			"LoadingPhase": "PostConfigInit",
+			"PlatformAllowList": [
+				"Android",
+				"IOS",
+				"Linux",
+				"LinuxArm64",
+				"Mac",
+				"Win64"
+			]
+		}
+	],
+	"Plugins": [
+		{
+			"Name": "OnlineServices",
+			"Enabled": true
+		},
+		{
+			"Name": "EOSShared",
+			"Enabled": true
+		}
+	]
+}
\ No newline at end of file
diff --git a/Engine/Plugins/Online/OnlineServicesEpicCommon/Source/OnlineServicesEpicCommon/OnlineServicesEpicCommon.Build.cs b/Engine/Plugins/Online/OnlineServicesEpicCommon/Source/OnlineServicesEpicCommon/OnlineServicesEpicCommon.Build.cs
new file mode 100644
--- /dev/null
+++ b/Engine/Plugins/Online/OnlineServicesEpicCommon/Source/OnlineServicesEpicCommon/OnlineServicesEpicCommon.Build.cs
@@ -0,0 +1,25 @@
+// Copyright Epic Games, Inc. All Rights Reserved.
+
+using UnrealBuildTool;
+using System.IO;
+
+public class OnlineServicesEpicCommon : ModuleRules
+{
+	public OnlineServicesEpicCommon(ReadOnlyTargetRules Target) : base(Target)
+	{
+		PublicDependencyModuleNames.AddRange(
+			new string[] {
+				"Core",
+				"OnlineServicesCommon"
+			}
+		);
+
+		PrivateDependencyModuleNames.AddRange(
+			new string[] {
+				"EOSSDK",
+				"EOSShared",
+				"OnlineServicesInterface"
+			}
+		);
+	}
+}
diff --git a/Engine/Plugins/Online/OnlineServicesEpicCommon/Source/OnlineServicesEpicCommon/Private/Online/EOSAuthLoginOptionsCommon.cpp b/Engine/Plugins/Online/OnlineServicesEpicCommon/Source/OnlineServicesEpicCommon/Private/Online/EOSAuthLoginOptionsCommon.cpp
new file mode 100644
--- /dev/null
+++ b/Engine/Plugins/Online/OnlineServicesEpicCommon/Source/OnlineServicesEpicCommon/Private/Online/EOSAuthLoginOptionsCommon.cpp
@@ -0,0 +1,110 @@
+// Copyright Epic Games, Inc. All Rights Reserved.
+
+#include "Online/EOSAuthLoginOptionsCommon.h"
+
+#include "EOSShared.h"
+
+namespace UE::Online {
+
+namespace CompatibilityLoginCredentialsType
+{
+	const FName Password = TEXT("epic");
+	const FName Developer = TEXT("dev_tool");
+}
+
+const FEOSAuthTranslationTraits* FEOSAuthLoginOptionsCommon::GetLoginTranslatorTraits(FName Name)
+{
+	static const TMap<FName, FEOSAuthTranslationTraits> SupportedLoginTranslatorTraits = {
+		{ LoginCredentialsType::Password, { EOS_ELoginCredentialType::EOS_LCT_Password, EEOSAuthTranslationFlags::SetId | EEOSAuthTranslationFlags::SetTokenFromString } },
+		{ LoginCredentialsType::ExchangeCode, { EOS_ELoginCredentialType::EOS_LCT_ExchangeCode, EEOSAuthTranslationFlags::SetTokenFromString } },
+		{ LoginCredentialsType::PersistentAuth, { EOS_ELoginCredentialType::EOS_LCT_PersistentAuth, EEOSAuthTranslationFlags::None } },
+		{ LoginCredentialsType::Developer, { EOS_ELoginCredentialType::EOS_LCT_Developer, EEOSAuthTranslationFlags::SetId | EEOSAuthTranslationFlags::SetTokenFromString } },
+		{ LoginCredentialsType::RefreshToken, { EOS_ELoginCredentialType::EOS_LCT_RefreshToken, EEOSAuthTranslationFlags::SetTokenFromString } },
+		{ LoginCredentialsType::AccountPortal, { EOS_ELoginCredentialType::EOS_LCT_AccountPortal, EEOSAuthTranslationFlags::SetId | EEOSAuthTranslationFlags::SetTokenFromString } },
+		{ LoginCredentialsType::ExternalAuth, { EOS_ELoginCredentialType::EOS_LCT_ExternalAuth, EEOSAuthTranslationFlags::SetTokenFromExternalAuth } },
+		{ CompatibilityLoginCredentialsType::Password, { EOS_ELoginCredentialType::EOS_LCT_Password, EEOSAuthTranslationFlags::SetId | EEOSAuthTranslationFlags::SetTokenFromString } },
+		{ CompatibilityLoginCredentialsType::Developer, { EOS_ELoginCredentialType::EOS_LCT_Developer, EEOSAuthTranslationFlags::SetId | EEOSAuthTranslationFlags::SetTokenFromString } },
+	};
+
+	return SupportedLoginTranslatorTraits.Find(Name);
+}
+
+const FEOSExternalAuthTranslationTraits* FEOSAuthLoginOptionsCommon::GetExternalAuthTranslationTraits(FName ExternalAuthType)
+{
+	static const TMap<FName, FEOSExternalAuthTranslationTraits> SupportedExternalAuthTraits = {
+		{ ExternalLoginType::Epic, { EOS_EExternalCredentialType::EOS_ECT_EPIC, EOS_ELinkAccountFlags::EOS_LA_NoFlags } },
+		{ ExternalLoginType::SteamSessionTicket, { EOS_EExternalCredentialType::EOS_ECT_STEAM_SESSION_TICKET, EOS_ELinkAccountFlags::EOS_LA_NoFlags } },
+		{ ExternalLoginType::PsnIdToken, { EOS_EExternalCredentialType::EOS_ECT_PSN_ID_TOKEN, EOS_ELinkAccountFlags::EOS_LA_NoFlags } },
+		{ ExternalLoginType::XblXstsToken, { EOS_EExternalCredentialType::EOS_ECT_XBL_XSTS_TOKEN, EOS_ELinkAccountFlags::EOS_LA_NoFlags } },
+		{ ExternalLoginType::DiscordAccessToken, { EOS_EExternalCredentialType::EOS_ECT_DISCORD_ACCESS_TOKEN, EOS_ELinkAccountFlags::EOS_LA_NoFlags } },
+		{ ExternalLoginType::GogSessionTicket, { EOS_EExternalCredentialType::EOS_ECT_GOG_SESSION_TICKET, EOS_ELinkAccountFlags::EOS_LA_NoFlags } },
+		{ ExternalLoginType::NintendoIdToken, { EOS_EExternalCredentialType::EOS_ECT_NINTENDO_ID_TOKEN, EOS_ELinkAccountFlags::EOS_LA_NoFlags } },
+		{ ExternalLoginType::NintendoNsaIdToken, { EOS_EExternalCredentialType::EOS_ECT_NINTENDO_NSA_ID_TOKEN, EOS_ELinkAccountFlags::EOS_LA_NintendoNsaId } },
+		{ ExternalLoginType::UplayAccessToken, { EOS_EExternalCredentialType::EOS_ECT_UPLAY_ACCESS_TOKEN, EOS_ELinkAccountFlags::EOS_LA_NoFlags } },
+		{ ExternalLoginType::OpenIdAccessToken, { EOS_EExternalCredentialType::EOS_ECT_OPENID_ACCESS_TOKEN, EOS_ELinkAccountFlags::EOS_LA_NoFlags } },
+		{ ExternalLoginType::DeviceIdAccessToken, { EOS_EExternalCredentialType::EOS_ECT_DEVICEID_ACCESS_TOKEN, EOS_ELinkAccountFlags::EOS_LA_NoFlags } },
+		{ ExternalLoginType::AppleIdToken, { EOS_EExternalCredentialType::EOS_ECT_APPLE_ID_TOKEN, EOS_ELinkAccountFlags::EOS_LA_NoFlags } },
+		{ ExternalLoginType::GoogleIdToken, { EOS_EExternalCredentialType::EOS_ECT_GOOGLE_ID_TOKEN, EOS_ELinkAccountFlags::EOS_LA_NoFlags } },
+		{ ExternalLoginType::OculusUserIdNonce, { EOS_EExternalCredentialType::EOS_ECT_OCULUS_USERID_NONCE, EOS_ELinkAccountFlags::EOS_LA_NoFlags } },
+		{ ExternalLoginType::ItchioJwt, { EOS_EExternalCredentialType::EOS_ECT_ITCHIO_JWT, EOS_ELinkAccountFlags::EOS_LA_NoFlags } },
+		{ ExternalLoginType::ItchioKey, { EOS_EExternalCredentialType::EOS_ECT_ITCHIO_KEY, EOS_ELinkAccountFlags::EOS_LA_NoFlags } },
+		{ ExternalLoginType::EpicIdToken, { EOS_EExternalCredentialType::EOS_ECT_EPIC_ID_TOKEN, EOS_ELinkAccountFlags::EOS_LA_NoFlags } },
+		{ ExternalLoginType::AmazonAccessToken, { EOS_EExternalCredentialType::EOS_ECT_AMAZON_ACCESS_TOKEN, EOS_ELinkAccountFlags::EOS_LA_NoFlags } },
+	};
+
+	return SupportedExternalAuthTraits.Find(ExternalAuthType);
+}
+
+FEOSAuthLoginOptionsCommon::FEOSAuthLoginOptionsCommon(FEOSAuthLoginOptionsCommon&& Other)
+{
+	*this = MoveTemp(Other);
+}
+
+FEOSAuthLoginOptionsCommon& FEOSAuthLoginOptionsCommon::operator=(FEOSAuthLoginOptionsCommon&& Other)
+{
+	CredentialsData = Other.CredentialsData;
+
+	// Pointer fixup.
+	if (CredentialsData.Id)
+	{
+		IdUtf8 = MoveTemp(Other.IdUtf8);
+		CredentialsData.Id = IdUtf8.GetData();
+	}
+	if (CredentialsData.Token)
+	{
+		TokenUtf8 = MoveTemp(Other.TokenUtf8);
+		CredentialsData.Token = TokenUtf8.GetData();
+		LinkAccountFlags = Other.LinkAccountFlags;
+	}
+	if (CredentialsData.SystemAuthCredentialsOptions)
+	{
+		// todo
+	}
+
+	Credentials = &CredentialsData;
+	ApiVersion = Other.ApiVersion;
+	ScopeFlags = Other.ScopeFlags;
+
+	Other.Credentials = nullptr;
+	return *this;
+}
+
+FEOSAuthLoginOptionsCommon::FEOSAuthLoginOptionsCommon()
+{
+	// EOS_Auth_LoginOptions init
+	UE_EOS_CHECK_API_MISMATCH(EOS_AUTH_LOGIN_API_LATEST, 3);
+	ApiVersion = 2;
+	Credentials = &CredentialsData;
+	ScopeFlags = EOS_EAuthScopeFlags::EOS_AS_NoFlags;
+	LoginFlags = 0;
+
+	// EOS_Auth_Credentials init
+	UE_EOS_CHECK_API_MISMATCH(EOS_AUTH_CREDENTIALS_API_LATEST, 4);
+	CredentialsData.ApiVersion = 4;
+	CredentialsData.Id = nullptr;
+	CredentialsData.Token = nullptr;
+	CredentialsData.Type = EOS_ELoginCredentialType::EOS_LCT_Password;
+	CredentialsData.SystemAuthCredentialsOptions = nullptr;
+	CredentialsData.ExternalType = EOS_EExternalCredentialType::EOS_ECT_EPIC;
+}
+}
\ No newline at end of file
diff --git a/Engine/Plugins/Online/OnlineServicesEpicCommon/Source/OnlineServicesEpicCommon/Private/Online/OnlineErrorEpicCommon.cpp b/Engine/Plugins/Online/OnlineServicesEpicCommon/Source/OnlineServicesEpicCommon/Private/Online/OnlineErrorEpicCommon.cpp
new file mode 100644
--- /dev/null
+++ b/Engine/Plugins/Online/OnlineServicesEpicCommon/Source/OnlineServicesEpicCommon/Private/Online/OnlineErrorEpicCommon.cpp
@@ -0,0 +1,85 @@
+// Copyright Epic Games, Inc. All Rights Reserved.
+
+#include "Online/OnlineErrorEpicCommon.h"
+#include "Online/OnlineErrorDefinitions.h"
+
+#include "EOSShared.h"
+
+namespace UE::Online::Errors {
+
+class FOnlineErrorDetailsEOS : public IOnlineErrorDetails
+{
+public:
+	static const TSharedRef<const IOnlineErrorDetails, ESPMode::ThreadSafe>& Get();
+
+	virtual FString GetFriendlyErrorCode(const FOnlineError& OnlineError) const override;
+
+	virtual FText GetText(const FOnlineError& OnlineError) const override;
+	virtual FString GetLogString(const FOnlineError& OnlineError) const override;
+};
+
+const TSharedRef<const IOnlineErrorDetails, ESPMode::ThreadSafe>& FOnlineErrorDetailsEOS::Get()
+{
+	static TSharedRef<const IOnlineErrorDetails, ESPMode::ThreadSafe> Instance = MakeShared<FOnlineErrorDetailsEOS>();
+	return Instance;
+}
+
+FString FOnlineErrorDetailsEOS::GetFriendlyErrorCode(const FOnlineError& OnlineError) const
+{
+	const EOS_EResult EosResult = EOS_EResult(OnlineError.GetValue());
+	return LexToString(EosResult);
+}
+
+FText FOnlineErrorDetailsEOS::GetText(const FOnlineError& OnlineError) const
+{
+	const EOS_EResult EosResult = EOS_EResult(OnlineError.GetValue());
+	return FText::FromString(LexToString(EosResult));
+}
+
+FString FOnlineErrorDetailsEOS::GetLogString(const FOnlineError& OnlineError) const
+{
+	const EOS_EResult EosResult = EOS_EResult(OnlineError.GetValue());
+	return LexToString(EosResult);
+}
+
+/** This doesn't wrap every single error in EOS, only the common ones that are strongly related to common errors */
+FOnlineError MapCommonEOSError(FOnlineError&& Error, EOS_EResult Result)
+{
+	switch (Result)
+	{
+	case EOS_EResult::EOS_Success:				return Errors::Success(MoveTemp(Error));
+	case EOS_EResult::EOS_NoConnection:			return Errors::NoConnection(MoveTemp(Error));
+	case EOS_EResult::EOS_InvalidCredentials:	return Errors::InvalidCreds(MoveTemp(Error));
+	case EOS_EResult::EOS_InvalidUser:			return Errors::InvalidUser(MoveTemp(Error));
+	case EOS_EResult::EOS_InvalidAuth:			return Errors::InvalidAuth(MoveTemp(Error));
+	case EOS_EResult::EOS_AccessDenied:			return Errors::AccessDenied(MoveTemp(Error));
+	case EOS_EResult::EOS_MissingPermissions:	return Errors::AccessDenied(MoveTemp(Error));
+	case EOS_EResult::EOS_TooManyRequests:		return Errors::TooManyRequests(MoveTemp(Error));
+	case EOS_EResult::EOS_AlreadyPending:		return Errors::AlreadyPending(MoveTemp(Error));
+	case EOS_EResult::EOS_InvalidParameters:	return Errors::InvalidParams(MoveTemp(Error));
+	case EOS_EResult::EOS_InvalidRequest:		return Errors::InvalidParams(MoveTemp(Error));
+	case EOS_EResult::EOS_UnrecognizedResponse: return Errors::InvalidResults(MoveTemp(Error));
+	case EOS_EResult::EOS_IncompatibleVersion:	return Errors::IncompatibleVersion(MoveTemp(Error));
+	case EOS_EResult::EOS_NotConfigured:		return Errors::NotConfigured(MoveTemp(Error));
+	case EOS_EResult::EOS_NotImplemented:		return Errors::NotImplemented(MoveTemp(Error));
+	case EOS_EResult::EOS_Canceled:				return Errors::Cancelled(MoveTemp(Error));
+	case EOS_EResult::EOS_NotFound:				return Errors::NotFound(MoveTemp(Error));
+	case EOS_EResult::EOS_OperationWillRetry:	return Errors::WillRetry(MoveTemp(Error));
+	case EOS_EResult::EOS_VersionMismatch:		return Errors::IncompatibleVersion(MoveTemp(Error));
+	case EOS_EResult::EOS_LimitExceeded:		return Errors::TooManyRequests(MoveTemp(Error));
+	case EOS_EResult::EOS_TimedOut:				return Errors::Timeout(MoveTemp(Error));
+	default:									return MoveTemp(Error);
+	}
+}
+
+ErrorCodeType ErrorCodeFromEOSResult(EOS_EResult EosResult)
+{
+	return Errors::ErrorCode::Create(Errors::ErrorCode::Category::EOS_System, Errors::ErrorCode::Category::EOS, uint32(EosResult));
+}
+
+FOnlineError FromEOSResult(EOS_EResult EosResult, FErrorMapperEosFn&& MapperFn)
+{
+	return MapperFn(FOnlineError(ErrorCodeFromEOSResult(EosResult), FOnlineErrorDetailsEOS::Get()), EosResult);
+}
+
+} /* namespace UE::Online::Errors */
diff --git a/Engine/Plugins/Online/OnlineServicesEpicCommon/Source/OnlineServicesEpicCommon/Private/Online/OnlineServicesEpicCommon.cpp b/Engine/Plugins/Online/OnlineServicesEpicCommon/Source/OnlineServicesEpicCommon/Private/Online/OnlineServicesEpicCommon.cpp
new file mode 100644
--- /dev/null
+++ b/Engine/Plugins/Online/OnlineServicesEpicCommon/Source/OnlineServicesEpicCommon/Private/Online/OnlineServicesEpicCommon.cpp
@@ -0,0 +1,72 @@
+// Copyright Epic Games, Inc. All Rights Reserved.
+
+#include "Online/OnlineServicesEpicCommon.h"
+
+#include "Online/OnlineServicesEpicCommonPlatformFactory.h"
+#include "Online/OnlineServicesLog.h"
+
+namespace UE::Online {
+
+#define UE_ONLINE_EOS_KEY_NAME_FAST_TICK_LOCK TEXT("FastTickLock")
+
+struct FOnlineServicesEpicCommonConfig
+{
+	bool bEnableAsyncOpFastTick = true;
+};
+
+namespace Meta {
+	BEGIN_ONLINE_STRUCT_META(FOnlineServicesEpicCommonConfig)
+	ONLINE_STRUCT_FIELD(FOnlineServicesEpicCommonConfig, bEnableAsyncOpFastTick)
+	END_ONLINE_STRUCT_META()
+}
+
+FOnlineServicesEpicCommon::FOnlineServicesEpicCommon(const FString& InServiceConfigName, FName InInstanceName, FName InInstanceConfigName)
+	: Super(InServiceConfigName, InInstanceName, InInstanceConfigName)
+{
+}
+
+bool FOnlineServicesEpicCommon::PreInit()
+{
+	FOnlineServicesEpicCommonPlatformFactory& PlatformFactory = FOnlineServicesEpicCommonPlatformFactory::Get();
+	EOSPlatformHandle = PlatformFactory.CreatePlatform(InstanceName, InstanceConfigName);
+	
+	if (!EOSPlatformHandle)
+	{
+		UE_LOG(LogOnlineServices, Warning, TEXT("[%hs] InstanceName=%s InstanceConfigName=%s EOSPlatformHandle=nullptr."), __FUNCTION__, *InstanceName.ToString(), *InstanceConfigName.ToString());
+		return false;
+	}
+
+	return true;
+}
+
+void FOnlineServicesEpicCommon::UpdateConfig()
+{
+	Super::UpdateConfig();
+
+	FOnlineServicesEpicCommonConfig Config;
+	LoadConfig(Config);
+	bEnableAsyncOpFastTick = Config.bEnableAsyncOpFastTick;
+}
+
+void FOnlineServicesEpicCommon::AddEOSSDKFastTick(FOnlineAsyncOp& InAsyncOp)
+{
+	if (bEnableAsyncOpFastTick)
+	{
+		if (EOSPlatformHandle)
+		{
+			TSharedPtr<IEOSFastTickLock> FastTickLock = EOSPlatformHandle->GetFastTickLock();
+			InAsyncOp.Data.Set<TSharedPtr<IEOSFastTickLock>>(UE_ONLINE_EOS_KEY_NAME_FAST_TICK_LOCK, FastTickLock);
+		}
+	}
+}
+
+void FOnlineServicesEpicCommon::RemoveEOSSDKFastTick(FOnlineAsyncOp& InAsyncOp)
+{
+	if (TSharedPtr<IEOSFastTickLock>* FastTickLock = InAsyncOp.Data.Get<TSharedPtr<IEOSFastTickLock>>(UE_ONLINE_EOS_KEY_NAME_FAST_TICK_LOCK))
+	{
+		// Simply reset the shared pointer. A null shared pointer may remain on the async op.
+		FastTickLock->Reset();
+	}
+}
+
+/* UE::Online */ }
diff --git a/Engine/Plugins/Online/OnlineServicesEpicCommon/Source/OnlineServicesEpicCommon/Private/Online/OnlineServicesEpicCommonModule.cpp b/Engine/Plugins/Online/OnlineServicesEpicCommon/Source/OnlineServicesEpicCommon/Private/Online/OnlineServicesEpicCommonModule.cpp
new file mode 100644
--- /dev/null
+++ b/Engine/Plugins/Online/OnlineServicesEpicCommon/Source/OnlineServicesEpicCommon/Private/Online/OnlineServicesEpicCommonModule.cpp
@@ -0,0 +1,29 @@
+// Copyright Epic Games, Inc. All Rights Reserved.
+
+#include "Online/OnlineServicesEpicCommonModule.h"
+
+#include "Modules/ModuleManager.h"
+#include "Online/OnlineServicesEpicCommon.h"
+#include "Online/OnlineServicesEpicCommonPlatformFactory.h"
+#include "Online/OnlineServicesRegistry.h"
+
+namespace UE::Online
+{
+
+void FOnlineServicesEpicCommonModule::StartupModule()
+{
+	FModuleManager::Get().LoadModuleChecked(TEXT("OnlineServicesInterface"));
+	FModuleManager::Get().LoadModuleChecked(TEXT("EOSShared"));
+
+	// Initialize the platform factory on startup.  This is necessary for the SDK to bind to rendering and input very early.
+	FOnlineServicesEpicCommonPlatformFactory::Get();
+}
+
+void FOnlineServicesEpicCommonModule::ShutdownModule()
+{
+	UE::Online::FOnlineServicesEpicCommonPlatformFactory::TearDown();
+}
+
+/* UE::Online */ }
+
+IMPLEMENT_MODULE(UE::Online::FOnlineServicesEpicCommonModule, OnlineServicesEpicCommon);
diff --git a/Engine/Plugins/Online/OnlineServicesEpicCommon/Source/OnlineServicesEpicCommon/Private/Online/OnlineServicesEpicCommonPlatformFactory.cpp b/Engine/Plugins/Online/OnlineServicesEpicCommon/Source/OnlineServicesEpicCommon/Private/Online/OnlineServicesEpicCommonPlatformFactory.cpp
new file mode 100644
--- /dev/null
+++ b/Engine/Plugins/Online/OnlineServicesEpicCommon/Source/OnlineServicesEpicCommon/Private/Online/OnlineServicesEpicCommonPlatformFactory.cpp
@@ -0,0 +1,147 @@
+// Copyright Epic Games, Inc. All Rights Reserved.
+
+#include "Online/OnlineServicesEpicCommonPlatformFactory.h"
+
+#include "EOSShared.h"
+#include "Online/OnlineExecHandler.h"
+#include "Online/OnlineServicesEpicCommon.h"
+
+#include "Misc/ConfigCacheIni.h"
+#include "Misc/Fork.h"
+#include "Misc/LazySingleton.h"
+#include "Modules/ModuleManager.h"
+
+#include "Online/OnlineServicesLog.h"
+#include <eos_base.h>
+
+namespace UE::Online {
+
+FOnlineServicesEpicCommonPlatformFactory& FOnlineServicesEpicCommonPlatformFactory::Get()
+{
+	return TLazySingleton<FOnlineServicesEpicCommonPlatformFactory>::Get();
+}
+
+void FOnlineServicesEpicCommonPlatformFactory::TearDown()
+{
+	return TLazySingleton<FOnlineServicesEpicCommonPlatformFactory>::TearDown();
+}
+
+IEOSPlatformHandlePtr FOnlineServicesEpicCommonPlatformFactory::CreatePlatform(FName InstanceName, FName InstanceConfigName)
+{
+	const FName EOSSharedModuleName = TEXT("EOSShared");
+	if (!FModuleManager::Get().IsModuleLoaded(EOSSharedModuleName))
+	{
+		FModuleManager::Get().LoadModuleChecked(EOSSharedModuleName);
+	}
+
+	IEOSSDKManager* const SDKManager = IEOSSDKManager::Get();
+	if (!SDKManager)
+	{
+		UE_LOG(LogOnlineServices, Error, TEXT("[FOnlineServicesEpicCommonPlatformFactory::CreatePlatform] EOSSDK has not been loaded."));
+		return {};
+	}
+
+	if (!SDKManager->IsInitialized())
+	{
+		UE_LOG(LogOnlineServices, Error, TEXT("[FOnlineServicesEpicCommonPlatformFactory::CreatePlatform] EOSSDK has not been initialized."));
+		return {};
+	}
+
+	FString PlatformConfigName;
+	if (!InstanceConfigName.IsNone())
+	{
+		PlatformConfigName = InstanceConfigName.ToString();
+	}
+	if (PlatformConfigName.IsEmpty())
+	{
+		PlatformConfigName = LoadEOSPlatformConfig(SDKManager);
+	}
+	if (PlatformConfigName.IsEmpty())
+	{
+		// Check for default platform config that other modules may have setup.
+		PlatformConfigName = SDKManager->GetDefaultPlatformConfigName();
+		if (PlatformConfigName.IsEmpty())
+		{
+			UE_LOG(LogOnlineServices, Verbose, TEXT("[FOnlineServicesEpicCommonPlatformFactory::CreatePlatform] Could not find platform config."));
+			return {};
+		}
+	}
+
+	IEOSPlatformHandlePtr EOSPlatformHandle = SDKManager->CreatePlatform(PlatformConfigName, InstanceName);
+	if (!EOSPlatformHandle)
+	{
+		UE_LOG(LogOnlineServices, Warning, TEXT("[FOnlineServicesEpicCommonPlatformFactory::CreatePlatform] Failed to create platform."));
+	}
+
+	return EOSPlatformHandle;
+}
+
+FString FOnlineServicesEpicCommonPlatformFactory::LoadEOSPlatformConfig(IEOSSDKManager* const SDKManager)
+{
+	// TODO: OnlineServices.EOS hardcoded until FURL support is added to the Services EpicGame and EpicAccount
+	FString ConfigSectionName = FString::Printf(TEXT("OnlineServices.EOS"));
+	if (!GConfig->DoesSectionExist(*ConfigSectionName, GEngineIni))
+	{
+		return FString();
+	}
+
+	// Check for explicit shared name to use.
+	FString PlatformConfigName;
+	GConfig->GetString(*ConfigSectionName, TEXT("PlatformConfigName"), PlatformConfigName, GEngineIni);
+	if (!PlatformConfigName.IsEmpty())
+	{
+		return PlatformConfigName;
+	}
+
+	// Check for cached config.
+	if (SDKManager->GetPlatformConfig(ConfigSectionName))
+	{
+		return ConfigSectionName;
+	}
+
+	// Check for legacy config. This should be handled by EOSShared instead, see FEOSSDKManager::GetPlatformConfig for details.
+	FEOSSDKPlatformConfig PlatformConfig;
+	PlatformConfig.Name = ConfigSectionName;
+
+	GConfig->GetString(*ConfigSectionName, TEXT("ProductId"), PlatformConfig.ProductId, GEngineIni);
+	if (PlatformConfig.ProductId.IsEmpty())
+	{
+		// If we're missing ProductId, assume we're missing the rest and instead rely on default EOSShared config.
+		return FString();
+	}
+
+	// Instead of specifying this config under the OnlineServices.EOS section, options should be moved to a new EOSSDK.Platform.<name> section so
+	// all modules relying on the EOSSDK can share the same config and platform instance. See FEOSSDKManager::GetPlatformConfig for details.
+	UE_LOG(LogOnlineServices, Warning, TEXT("[LoadEOSPlatformConfig] Using legacy config from %s, use EOSShared named config instead."), *ConfigSectionName);
+
+	GConfig->GetString(*ConfigSectionName, TEXT("SandboxId"), PlatformConfig.SandboxId, GEngineIni);
+	GConfig->GetString(*ConfigSectionName, TEXT("DeploymentId"), PlatformConfig.DeploymentId, GEngineIni);
+	GConfig->GetString(*ConfigSectionName, TEXT("ClientId"), PlatformConfig.ClientId, GEngineIni);
+	GConfig->GetString(*ConfigSectionName, TEXT("ClientSecret"), PlatformConfig.ClientSecret, GEngineIni);
+	// Config key renamed to ClientEncryptionKey as EncryptionKey gets removed from packaged builds due to IniKeyDenylist=EncryptionKey entry in BaseGame.ini.
+	GConfig->GetString(*ConfigSectionName, TEXT("ClientEncryptionKey"), PlatformConfig.EncryptionKey, GEngineIni);
+
+	const FString CacheDirBase = SDKManager->GetCacheDirBase();
+	PlatformConfig.CacheDirectory = CacheDirBase.IsEmpty() ? FString() : CacheDirBase / TEXT("OnlineServicesEOS");
+
+	PlatformConfig.bIsServer = IsRunningDedicatedServer() ? EOS_TRUE : EOS_FALSE;
+	if (!IsRunningGame())
+	{
+		PlatformConfig.bLoadingInEditor = true;
+	}
+	else
+	{
+		PlatformConfig.bWindowsEnableOverlayD3D9 = true;
+		PlatformConfig.bWindowsEnableOverlayD3D10 = true;
+		PlatformConfig.bWindowsEnableOverlayOpenGL = true;
+	}
+
+	if (!SDKManager->AddPlatformConfig(PlatformConfig))
+	{
+		return FString();
+	}
+
+	return ConfigSectionName;
+}
+
+/* UE::Online */ }
diff --git a/Engine/Plugins/Online/OnlineServicesEpicCommon/Source/OnlineServicesEpicCommon/Public/Online/EOSAuthLoginOptionsCommon.h b/Engine/Plugins/Online/OnlineServicesEpicCommon/Source/OnlineServicesEpicCommon/Public/Online/EOSAuthLoginOptionsCommon.h
new file mode 100644
--- /dev/null
+++ b/Engine/Plugins/Online/OnlineServicesEpicCommon/Source/OnlineServicesEpicCommon/Public/Online/EOSAuthLoginOptionsCommon.h
@@ -0,0 +1,61 @@
+// Copyright Epic Games, Inc. All Rights Reserved.
+
+#pragma once
+
+#include "Online/Auth.h"
+#include "Online/OnlineAsyncOp.h"
+
+#if defined(EOS_PLATFORM_BASE_FILE_NAME)
+#include EOS_PLATFORM_BASE_FILE_NAME
+#endif
+
+#include "eos_auth_types.h"
+
+namespace UE::Online {
+
+enum class EEOSAuthTranslationFlags : uint8
+{
+	None = 0,
+	SetId = 1 << 0,
+	SetTokenFromString = 1 << 1,
+	SetTokenFromExternalAuth = 1 << 2,
+};
+ENUM_CLASS_FLAGS(EEOSAuthTranslationFlags);
+
+struct FEOSAuthTranslationTraits
+{
+	EOS_ELoginCredentialType Type = EOS_ELoginCredentialType::EOS_LCT_Password;
+	EEOSAuthTranslationFlags Flags = EEOSAuthTranslationFlags::None;
+};
+
+struct FEOSExternalAuthTranslationTraits
+{
+	EOS_EExternalCredentialType Type = EOS_EExternalCredentialType::EOS_ECT_EPIC;
+	EOS_ELinkAccountFlags LinkAccountFlags = EOS_ELinkAccountFlags::EOS_LA_NoFlags;
+};
+
+class ONLINESERVICESEPICCOMMON_API FEOSAuthLoginOptionsCommon : public EOS_Auth_LoginOptions
+{
+public:
+	FEOSAuthLoginOptionsCommon(const FEOSAuthLoginOptionsCommon&) = delete;
+	FEOSAuthLoginOptionsCommon& operator=(const FEOSAuthLoginOptionsCommon&) = delete;
+	FEOSAuthLoginOptionsCommon(FEOSAuthLoginOptionsCommon&&);
+	FEOSAuthLoginOptionsCommon& operator=(FEOSAuthLoginOptionsCommon&&);
+
+	EOS_ELinkAccountFlags GetLinkAccountFlags() const { return LinkAccountFlags; }
+
+protected:
+	FEOSAuthLoginOptionsCommon();
+
+	static const FEOSAuthTranslationTraits* GetLoginTranslatorTraits(FName Name);
+	static const FEOSExternalAuthTranslationTraits* GetExternalAuthTranslationTraits(FName ExternalAuthType);
+
+	static bool InitSystemAuthCredentialOptions(FEOSAuthLoginOptionsCommon& Options) { return true; }
+
+	EOS_Auth_Credentials CredentialsData;
+	EOS_ELinkAccountFlags LinkAccountFlags = EOS_ELinkAccountFlags::EOS_LA_NoFlags;
+	TArray<char> IdUtf8;
+	TArray<char> TokenUtf8;
+};
+
+}
\ No newline at end of file
diff --git a/Engine/Plugins/Online/OnlineServicesEpicCommon/Source/OnlineServicesEpicCommon/Public/Online/OnlineErrorEpicCommon.h b/Engine/Plugins/Online/OnlineServicesEpicCommon/Source/OnlineServicesEpicCommon/Public/Online/OnlineErrorEpicCommon.h
new file mode 100644
--- /dev/null
+++ b/Engine/Plugins/Online/OnlineServicesEpicCommon/Source/OnlineServicesEpicCommon/Public/Online/OnlineErrorEpicCommon.h
@@ -0,0 +1,46 @@
+// Copyright Epic Games, Inc. All Rights Reserved.
+
+#pragma once
+
+#include "EOSShared.h"
+#include "Online/OnlineError.h"
+
+/*
+ *	Proper usage:
+ *		Errors::FromEOSResult(EOSResult::EOS_PlayerDataStorage_FileSizeTooLarge);
+ *	Certain EOS errors are predefined to have a common error parent type (see Internal_EOSWrapInner). i.e.
+ *		(Errors::FromEOSResult(EOSResult::EOS_NoConnection) == Errors::NoConnection()) == true
+ */
+
+namespace UE::Online::Errors {
+
+	UE_ONLINE_ERROR_CATEGORY(EOS, ThirdPartyPlugin, 0x4, "EOS")
+
+	typedef TFunction<FOnlineError(FOnlineError&& Error, EOS_EResult Result)> FErrorMapperEosFn;
+
+	ONLINESERVICESEPICCOMMON_API FOnlineError MapCommonEOSError(FOnlineError&& Error, EOS_EResult Result);
+	ONLINESERVICESEPICCOMMON_API ErrorCodeType ErrorCodeFromEOSResult(EOS_EResult Result);
+	ONLINESERVICESEPICCOMMON_API FOnlineError FromEOSResult(EOS_EResult Result, FErrorMapperEosFn&& MapperFn = &MapCommonEOSError);
+
+/* UE::Online::Errors */ }
+
+// These are extern'd out of the namespace for some of the Catch API that can't see these operators if they are inside
+inline bool operator==(const UE::Online::FOnlineError& OnlineError, EOS_EResult EosResult)
+{
+	return OnlineError == UE::Online::Errors::ErrorCodeFromEOSResult(EosResult);
+}
+
+inline bool operator==(EOS_EResult EosResult, const UE::Online::FOnlineError& OnlineError)
+{
+	return OnlineError == EosResult;
+}
+
+inline bool operator!=(const UE::Online::FOnlineError& OnlineError, EOS_EResult EosResult)
+{
+	return !(OnlineError == EosResult);
+}
+
+inline bool operator!=(EOS_EResult EosResult, const UE::Online::FOnlineError& OnlineError)
+{
+	return OnlineError != EosResult;
+}
diff --git a/Engine/Plugins/Online/OnlineServicesEpicCommon/Source/OnlineServicesEpicCommon/Public/Online/OnlineServicesEpicCommon.h b/Engine/Plugins/Online/OnlineServicesEpicCommon/Source/OnlineServicesEpicCommon/Public/Online/OnlineServicesEpicCommon.h
new file mode 100644
--- /dev/null
+++ b/Engine/Plugins/Online/OnlineServicesEpicCommon/Source/OnlineServicesEpicCommon/Public/Online/OnlineServicesEpicCommon.h
@@ -0,0 +1,35 @@
+// Copyright Epic Games, Inc. All Rights Reserved.
+
+#pragma once
+
+#include "Online/OnlineServicesCommon.h"
+
+using IEOSPlatformHandlePtr = TSharedPtr<class IEOSPlatformHandle>;
+
+namespace UE::Online {
+
+class ONLINESERVICESEPICCOMMON_API FOnlineServicesEpicCommon : public FOnlineServicesCommon
+{
+public:
+	using Super = FOnlineServicesCommon;
+
+	FOnlineServicesEpicCommon(const FString& InServiceConfigName, FName InInstanceName, FName InInstanceConfigName);
+	virtual ~FOnlineServicesEpicCommon() = default;
+
+	virtual bool PreInit();
+	virtual void UpdateConfig() override;
+
+	IEOSPlatformHandlePtr GetEOSPlatformHandle() const { return EOSPlatformHandle; }
+
+	/** Enable EOSSDK to tick as fast as it can while this operation is outstanding. Expected usage is to call this before calling an SDK function. */
+	void AddEOSSDKFastTick(FOnlineAsyncOp& InAsyncOp);
+	/** Remove fast ticking of EOSSDK for an operation. Expected usage is to call this after an SDK function's completion delegate triggers. If this is not called, the fast tick will be removed when the operation destructs. */
+	void RemoveEOSSDKFastTick(FOnlineAsyncOp& InAsyncOp);
+
+protected:
+	IEOSPlatformHandlePtr EOSPlatformHandle;
+
+	bool bEnableAsyncOpFastTick = true;
+};
+
+/* UE::Online */ }
diff --git a/Engine/Plugins/Online/OnlineServicesEpicCommon/Source/OnlineServicesEpicCommon/Public/Online/OnlineServicesEpicCommonModule.h b/Engine/Plugins/Online/OnlineServicesEpicCommon/Source/OnlineServicesEpicCommon/Public/Online/OnlineServicesEpicCommonModule.h
new file mode 100644
--- /dev/null
+++ b/Engine/Plugins/Online/OnlineServicesEpicCommon/Source/OnlineServicesEpicCommon/Public/Online/OnlineServicesEpicCommonModule.h
@@ -0,0 +1,18 @@
+// Copyright Epic Games, Inc. All Rights Reserved.
+
+#pragma once
+
+#include "Modules/ModuleInterface.h"
+
+namespace UE::Online {
+
+class FOnlineServicesEpicCommonModule : public IModuleInterface
+{
+private:
+	// Begin IModuleInterface
+	virtual void StartupModule() override;
+	virtual void ShutdownModule() override;
+	// End IModuleInterface
+};
+
+} // namespace UE::Online
\ No newline at end of file
diff --git a/Engine/Plugins/Online/OnlineServicesEpicCommon/Source/OnlineServicesEpicCommon/Public/Online/OnlineServicesEpicCommonPlatformFactory.h b/Engine/Plugins/Online/OnlineServicesEpicCommon/Source/OnlineServicesEpicCommon/Public/Online/OnlineServicesEpicCommonPlatformFactory.h
new file mode 100644
--- /dev/null
+++ b/Engine/Plugins/Online/OnlineServicesEpicCommon/Source/OnlineServicesEpicCommon/Public/Online/OnlineServicesEpicCommonPlatformFactory.h
@@ -0,0 +1,47 @@
+// Copyright Epic Games, Inc. All Rights Reserved.
+
+#pragma once
+
+#include "Templates/SharedPointer.h"
+
+#include "IEOSSDKManager.h"
+
+using IEOSPlatformHandlePtr = TSharedPtr<class IEOSPlatformHandle>;
+class FLazySingleton;
+
+namespace UE::Online
+{
+
+	/** Factory class to create EOS Platforms for online services */
+	class FOnlineServicesEpicCommonPlatformFactory
+	{
+	public:
+		/**
+		 * Get the platform factory
+		 * @return the platform factory singleton
+		 */
+		static FOnlineServicesEpicCommonPlatformFactory& Get();
+		/**
+		 * Tear down the singleton instance. This only cleans up the singleton and has no impact on any platform handles created by this (aside from DefaultEOSPlatformHandle's ref count decreasing).
+		 */
+		static void TearDown();
+
+		/**
+		 * Create a new platform instance for the given Instance/Config pair.
+		 * Note If InstanceConfigName = NAME_None, this will attempt to resolve a config to use from various sources, including OnlineServices config, and EOSSDKManager cached configs.
+		 * @param InstanceName the "instance", typically a WorldContextHandle
+		 * @param InstanceConfigName the named EOS platform config to use
+		 * @return a new platform, or null on failure
+		 */
+		IEOSPlatformHandlePtr CreatePlatform(FName InstanceName, FName InstanceConfigName);
+
+	private:
+		FOnlineServicesEpicCommonPlatformFactory() = default;
+
+		FString LoadEOSPlatformConfig(IEOSSDKManager* const SDKManager);
+
+		friend FLazySingleton;
+	};
+
+	/* UE::Online */
+}
