- Fix EOSGS shadowing the platform handle from EpicCommon

#rnx

diff --git a/Engine/Plugins/Online/OnlineServicesEOSGS/Source/Public/Online/OnlineServicesEOSGS.h b/Engine/Plugins/Online/OnlineServicesEOSGS/Source/Public/Online/OnlineServicesEOSGS.h
--- a/Engine/Plugins/Online/OnlineServicesEOSGS/Source/Public/Online/OnlineServicesEOSGS.h
+++ b/Engine/Plugins/Online/OnlineServicesEOSGS/Source/Public/Online/OnlineServicesEOSGS.h
@@ -31,17 +31,11 @@ public:
 	virtual TOnlineResult<FGetResolvedConnectString> GetResolvedConnectString(FGetResolvedConnectString::Params&& Params) override;
 	virtual EOnlineServices GetServicesProvider() const override { return EOnlineServices::Epic; }
 
-	IEOSPlatformHandlePtr GetEOSPlatformHandle() const { return EOSPlatformHandle; }
-
 	UE_DEPRECATED(5.5, "Please call GetServiceConfigNameStatic instead")
 	static const TCHAR* GetConfigNameStatic() { return GetServiceConfigNameStatic(); }
 	static const TCHAR* GetServiceConfigNameStatic() { return TEXT("EOS"); }
 
-
-
 protected:
-	IEOSPlatformHandlePtr EOSPlatformHandle;
-
 	void WarnIfEncryptionKeyMissing(const FString& InterfaceName) const;
 
 #if WITH_ENGINE
