Addition of missing include to fix build errors in projects updating Engine version.

diff --git a/Engine/Plugins/Online/OnlineServicesEOSGS/Source/Public/Online/OnlineServicesEOSGSTypes.h b/Engine/Plugins/Online/OnlineServicesEOSGS/Source/Public/Online/OnlineServicesEOSGSTypes.h
--- a/Engine/Plugins/Online/OnlineServicesEOSGS/Source/Public/Online/OnlineServicesEOSGSTypes.h
+++ b/Engine/Plugins/Online/OnlineServicesEOSGS/Source/Public/Online/OnlineServicesEOSGSTypes.h
@@ -4,4 +4,6 @@
 
 #include "EOSSharedTypes.h"
 
+#include "Online/OnlineErrorEpicCommon.h"
+
 #define LOBBIESEOSGS_LOBBY_ID_OVERRIDE_ATTRIBUTE_KEY TEXT("LOBBY_ID_OVERRIDE")
