Adding EOS Player Reports and Player Sanctions to OnlineServicesEOSGS

diff --git a/Engine/Plugins/Online/OnlineServicesEOSGS/Source/Private/Online/OnlineServicesEOSGS.cpp b/Engine/Plugins/Online/OnlineServicesEOSGS/Source/Private/Online/OnlineServicesEOSGS.cpp
--- a/Engine/Plugins/Online/OnlineServicesEOSGS/Source/Private/Online/OnlineServicesEOSGS.cpp
+++ b/Engine/Plugins/Online/OnlineServicesEOSGS/Source/Private/Online/OnlineServicesEOSGS.cpp
@@ -12,6 +12,8 @@
 #include "Online/SessionsEOSGS.h"
 #include "Online/TitleFileEOSGS.h"
 #include "Online/UserFileEOSGS.h"
+#include "Online/PlayerReportsEOSGS.h"
+#include "Online/PlayerSanctionsEOSGS.h"
 #include "IEOSSDKManager.h"
 
 #if WITH_ENGINE
@@ -92,6 +94,8 @@ void FOnlineServicesEOSGS::RegisterComponents()
 	Components.Register<FLobbiesEOSGS>(*this);
 	Components.Register<FStatsEOSGS>(*this);
 	Components.Register<FSessionsEOSGS>(*this);
+	Components.Register<FPlayerReportsEOSGS>(*this);
+	Components.Register<FPlayerSanctionsEOSGS>(*this);
 	if (EOS_Platform_GetTitleStorageInterface(*GetEOSPlatformHandle()))
 	{
 		Components.Register<FTitleFileEOSGS>(*this);
@@ -173,4 +177,15 @@ TOnlineResult<FGetResolvedConnectString> FOnlineServicesEOSGS::GetResolvedConnec
 	return TOnlineResult<FGetResolvedConnectString>(Errors::InvalidParams());
 }
 
+// Get specific OnlineServiesEOSGS interfaces
+IPlayerReportsPtr FOnlineServicesEOSGS::GetPlayerReportsInterface()
+{
+	return IPlayerReportsPtr(AsShared(), Get<IPlayerReports>());
+}
+
+IPlayerSanctionsPtr FOnlineServicesEOSGS::GetPlayerSanctionsInterface()
+{
+	return IPlayerSanctionsPtr(AsShared(), Get<IPlayerSanctions>());
+}
+
 /* UE::Online */ }
diff --git a/Engine/Plugins/Online/OnlineServicesEOSGS/Source/Private/Online/OnlineServicesEOSGSInterfaces/PlayerReports.cpp b/Engine/Plugins/Online/OnlineServicesEOSGS/Source/Private/Online/OnlineServicesEOSGSInterfaces/PlayerReports.cpp
new file mode 100644
--- /dev/null
+++ b/Engine/Plugins/Online/OnlineServicesEOSGS/Source/Private/Online/OnlineServicesEOSGSInterfaces/PlayerReports.cpp
@@ -0,0 +1,61 @@
+// Copyright Epic Games, Inc. All Rights Reserved.
+
+#include "Online/OnlineServicesEOSGSInterfaces/PlayerReports.h"
+
+namespace UE::Online {
+
+const TCHAR* LexToString(EPlayerReportCategory Value)
+{
+	switch (Value)
+	{
+	case EPlayerReportCategory::Cheating:
+		return TEXT("Cheating");
+	case EPlayerReportCategory::Exploiting:
+		return TEXT("Exploiting");
+	case EPlayerReportCategory::OffensiveProfile:
+		return TEXT("OffensiveProfile");
+	case EPlayerReportCategory::VerbalAbuse:
+		return TEXT("VerbalAbuse");
+	case EPlayerReportCategory::Scamming:
+		return TEXT("Scamming");
+	case EPlayerReportCategory::Spamming:
+		return TEXT("Spamming");
+	default: checkNoEntry();
+	case EPlayerReportCategory::Other:
+		return TEXT("Other");
+	}
+}
+
+void LexFromString(EPlayerReportCategory& OutValue, const TCHAR* InStr)
+{
+	if (FCString::Stricmp(InStr, TEXT("Cheating")) == 0)
+	{
+		OutValue = EPlayerReportCategory::Cheating;
+	}
+	else if (FCString::Stricmp(InStr, TEXT("Exploiting")) == 0)
+	{
+		OutValue = EPlayerReportCategory::Exploiting; 
+	}
+	else if (FCString::Stricmp(InStr, TEXT("OffensiveProfile")) == 0)
+	{
+		OutValue = EPlayerReportCategory::OffensiveProfile;
+	}
+	else if (FCString::Stricmp(InStr, TEXT("VerbalAbuse")) == 0)
+	{
+		OutValue = EPlayerReportCategory::VerbalAbuse;
+	}
+	else if (FCString::Stricmp(InStr, TEXT("Scamming")) == 0)
+	{
+		OutValue = EPlayerReportCategory::Scamming;
+	}
+	else if (FCString::Stricmp(InStr, TEXT("Spamming")) == 0)
+	{
+		OutValue = EPlayerReportCategory::Spamming;
+	}
+	else
+	{
+		OutValue = EPlayerReportCategory::Other; 
+	}
+}
+
+/* UE::Online */}
\ No newline at end of file
diff --git a/Engine/Plugins/Online/OnlineServicesEOSGS/Source/Private/Online/OnlineServicesEOSGSInterfaces/PlayerSanctions.cpp b/Engine/Plugins/Online/OnlineServicesEOSGS/Source/Private/Online/OnlineServicesEOSGSInterfaces/PlayerSanctions.cpp
new file mode 100644
--- /dev/null
+++ b/Engine/Plugins/Online/OnlineServicesEOSGS/Source/Private/Online/OnlineServicesEOSGSInterfaces/PlayerSanctions.cpp
@@ -0,0 +1,43 @@
+// Copyright Epic Games, Inc. All Rights Reserved.
+
+#include "Online/OnlineServicesEOSGSInterfaces/PlayerSanctions.h"
+
+namespace UE::Online {
+
+const TCHAR* LexToString(EPlayerSanctionAppealReason Value)
+{
+	switch (Value)
+	{
+	case EPlayerSanctionAppealReason::IncorrectSanction:
+		return TEXT("IncorrectSanction");
+	case EPlayerSanctionAppealReason::CompromisedAccount:
+		return TEXT("CompromisedAccount");
+	case EPlayerSanctionAppealReason::UnfairPunishment:
+		return TEXT("UnfairPunishment");
+	default: checkNoEntry();
+	case EPlayerSanctionAppealReason::AppealForForgiveness:
+		return TEXT("AppealForForgiveness");
+	}
+}
+
+void LexFromString(EPlayerSanctionAppealReason& OutValue, const TCHAR* InStr)
+{
+	if (FCString::Stricmp(InStr, TEXT("IncorrectSanction")) == 0)
+	{
+		OutValue = EPlayerSanctionAppealReason::IncorrectSanction;
+	}
+	else if (FCString::Stricmp(InStr, TEXT("CompromisedAccount")) == 0)
+	{
+		OutValue = EPlayerSanctionAppealReason::CompromisedAccount;
+	}
+	else if (FCString::Stricmp(InStr, TEXT("UnfairPunishment")) == 0)
+	{
+		OutValue = EPlayerSanctionAppealReason::UnfairPunishment;
+	}
+	else if (FCString::Stricmp(InStr, TEXT("AppealForForgiveness")) == 0)
+	{
+		OutValue = EPlayerSanctionAppealReason::AppealForForgiveness;
+	}
+}
+	
+/* UE::Online */}
\ No newline at end of file
diff --git a/Engine/Plugins/Online/OnlineServicesEOSGS/Source/Private/Online/PlayerReportsEOSGS.cpp b/Engine/Plugins/Online/OnlineServicesEOSGS/Source/Private/Online/PlayerReportsEOSGS.cpp
new file mode 100644
--- /dev/null
+++ b/Engine/Plugins/Online/OnlineServicesEOSGS/Source/Private/Online/PlayerReportsEOSGS.cpp
@@ -0,0 +1,124 @@
+// Copyright Epic Games, Inc. All Rights Reserved.
+
+#include "Online/PlayerReportsEOSGS.h"
+
+#include "EOSShared.h"
+#include "IEOSSDKManager.h"
+#include "Online/AuthEOSGS.h"
+#include "Online/OnlineErrorEpicCommon.h"
+#include "Online/OnlineIdEOSGS.h"
+#include "Online/OnlineServicesEOSGS.h"
+
+#include "eos_reports.h"
+
+namespace UE::Online {
+	
+FPlayerReportsEOSGS::FPlayerReportsEOSGS(FOnlineServicesEOSGS& InServices) : TOnlineComponent(TEXT("PlayerReports"), InServices)
+{
+}
+
+void FPlayerReportsEOSGS::Initialize()
+{
+	Super::Initialize(); 
+
+	PlayerReportsHandle = EOS_Platform_GetReportsInterface(*static_cast<FOnlineServicesEOSGS&>(GetServices()).GetEOSPlatformHandle());
+	check(PlayerReportsHandle);
+
+	// No need to register for any EOS notifications/events
+
+	RegisterCommands(); 
+}
+
+void FPlayerReportsEOSGS::RegisterCommands()
+{
+	RegisterCommand(&FPlayerReportsEOSGS::SendPlayerReport);
+}
+
+// Convert the plugin player report category enum to the EOS SDK enum
+EOS_EPlayerReportsCategory ToEOSPlayerReportsCategory(EPlayerReportCategory Category)
+{
+	switch (Category)
+	{
+	case EPlayerReportCategory::Cheating:
+	{
+		return EOS_EPlayerReportsCategory::EOS_PRC_Cheating;
+	}
+	case EPlayerReportCategory::Exploiting:
+	{
+		return EOS_EPlayerReportsCategory::EOS_PRC_Exploiting;
+	}
+	case EPlayerReportCategory::OffensiveProfile:
+	{
+		return EOS_EPlayerReportsCategory::EOS_PRC_OffensiveProfile;
+	}
+	case EPlayerReportCategory::VerbalAbuse:
+	{
+		return EOS_EPlayerReportsCategory::EOS_PRC_VerbalAbuse;
+	}
+	case EPlayerReportCategory::Scamming:
+	{
+		return EOS_EPlayerReportsCategory::EOS_PRC_Scamming;
+	}
+	case EPlayerReportCategory::Spamming:
+	{
+		return EOS_EPlayerReportsCategory::EOS_PRC_Spamming;
+	}
+	}
+	// Default to other
+	return EOS_EPlayerReportsCategory::EOS_PRC_Other;
+}
+
+TOnlineAsyncOpHandle<FSendPlayerReport> FPlayerReportsEOSGS::SendPlayerReport(FSendPlayerReport::Params&& Params)
+{
+	TOnlineAsyncOpRef<FSendPlayerReport> Op = GetOp<FSendPlayerReport>(MoveTemp(Params));
+	if (!Op->IsReady())
+	{
+		Op->Then([this](TOnlineAsyncOp<FSendPlayerReport>& InAsyncOp)
+		{
+			const FSendPlayerReport::Params& Params = InAsyncOp.GetParams();
+
+			if (!Services.Get<FAuthEOSGS>()->IsLoggedIn(Params.LocalAccountId))
+			{
+				InAsyncOp.SetError(Errors::InvalidUser());
+			}
+		})
+		.Then([this](TOnlineAsyncOp<FSendPlayerReport>& InAsyncOp, TPromise<const EOS_Reports_SendPlayerBehaviorReportCompleteCallbackInfo*>&& Promise)
+		{
+			const FSendPlayerReport::Params& Params = InAsyncOp.GetParams();
+
+			EOS_Reports_SendPlayerBehaviorReportOptions Options = {}; 
+			Options.ApiVersion = 2;
+			UE_EOS_CHECK_API_MISMATCH(EOS_REPORTS_SENDPLAYERBEHAVIORREPORT_API_LATEST, 2);
+			Options.ReporterUserId = GetProductUserIdChecked(Params.LocalAccountId);
+			Options.ReportedUserId = GetProductUserIdChecked(Params.TargetAccountId);
+			Options.Category = ToEOSPlayerReportsCategory(Params.Category);
+			const auto MessagedUTF8 = StringCast<UTF8CHAR>(*Params.Message);
+			Options.Message =Params.Message.IsEmpty() ? nullptr : (const char*)MessagedUTF8.Get();
+			const auto ContextUTF8 = StringCast<UTF8CHAR>(*Params.Context);
+			Options.Context = Params.Context.IsEmpty() ? nullptr : (const char*)ContextUTF8.Get();
+
+			EOS_Async(EOS_Reports_SendPlayerBehaviorReport, PlayerReportsHandle, Options, MoveTemp(Promise));
+		})
+		.Then([this](TOnlineAsyncOp<FSendPlayerReport>& InAsyncOp, const EOS_Reports_SendPlayerBehaviorReportCompleteCallbackInfo* Data)
+		{
+			if (Data->ResultCode != EOS_EResult::EOS_Success)
+			{
+				UE_LOG(LogOnlineServices, Warning, TEXT("EOS_Reports_SendPlayerBehaviorReport result=[%s]"), *LexToString(Data->ResultCode));
+				InAsyncOp.SetError(Errors::FromEOSResult(Data->ResultCode));
+			}
+			else
+			{
+				UE_LOG(LogOnlineServices, Verbose, TEXT("EOS_Reports_SendPlayerBehaviorReport result=[%s]"), *LexToString(Data->ResultCode));
+				InAsyncOp.SetResult({});
+			}
+		}).Enqueue(GetSerialQueue());
+	}
+	return Op->GetHandle();
+}
+
+FString ToLogString(const FSendPlayerReport::Result& SendPlayerReportResult)
+{
+	// SendPlayerReport Result is empty 
+	return FString();
+}
+/* UE::Online */ }
\ No newline at end of file
diff --git a/Engine/Plugins/Online/OnlineServicesEOSGS/Source/Private/Online/PlayerSanctionsEOSGS.cpp b/Engine/Plugins/Online/OnlineServicesEOSGS/Source/Private/Online/PlayerSanctionsEOSGS.cpp
new file mode 100644
--- /dev/null
+++ b/Engine/Plugins/Online/OnlineServicesEOSGS/Source/Private/Online/PlayerSanctionsEOSGS.cpp
@@ -0,0 +1,202 @@
+// Copyright Epic Games, Inc. All Rights Reserved.
+
+#include "Online/PlayerSanctionsEOSGS.h"
+
+#include "EOSShared.h"
+#include "IEOSSDKManager.h"
+#include "Online/AuthEOSGS.h"
+#include "Online/OnlineErrorEpicCommon.h"
+#include "Online/OnlineIdEOSGS.h"
+#include "Online/OnlineServicesEOSGS.h"
+
+#include "eos_sanctions.h"
+
+namespace UE::Online {
+	
+FPlayerSanctionsEOSGS::FPlayerSanctionsEOSGS(FOnlineServicesEOSGS& InServices) : TOnlineComponent(TEXT("PlayerSanctions"), InServices)
+{
+}
+
+void FPlayerSanctionsEOSGS::Initialize()
+{
+	Super::Initialize(); 
+
+	PlayerSanctionsHandle = EOS_Platform_GetSanctionsInterface(*static_cast<FOnlineServicesEOSGS&>(GetServices()).GetEOSPlatformHandle());
+	check(PlayerSanctionsHandle);
+
+	// No need to register for any EOS notifications/events
+
+	RegisterCommands();
+}
+
+void FPlayerSanctionsEOSGS::RegisterCommands()
+{
+	RegisterCommand(&FPlayerSanctionsEOSGS::CreatePlayerSanctionAppeal);
+	RegisterCommand(&FPlayerSanctionsEOSGS::ReadEntriesForUser);
+}
+
+// Convert the plugin player sanction category enum to the EOS SDK enum
+EOS_ESanctionAppealReason ToEOSSanctionAppealReason(EPlayerSanctionAppealReason Reason)
+{
+	switch (Reason)
+	{
+	case EPlayerSanctionAppealReason::IncorrectSanction:
+		return EOS_ESanctionAppealReason::EOS_SAR_IncorrectSanction;
+	case EPlayerSanctionAppealReason::CompromisedAccount:
+		return EOS_ESanctionAppealReason::EOS_SAR_CompromisedAccount;
+	case EPlayerSanctionAppealReason::UnfairPunishment:
+		return EOS_ESanctionAppealReason::EOS_SAR_UnfairPunishment;
+	case EPlayerSanctionAppealReason::AppealForForgiveness:
+		return EOS_ESanctionAppealReason::EOS_SAR_AppealForForgiveness;
+	}
+	return EOS_ESanctionAppealReason::EOS_SAR_Invalid;
+}
+
+TOnlineAsyncOpHandle<FCreatePlayerSanctionAppeal> FPlayerSanctionsEOSGS::CreatePlayerSanctionAppeal(FCreatePlayerSanctionAppeal::Params&& Params)
+{
+	TOnlineAsyncOpRef<FCreatePlayerSanctionAppeal> Op = GetOp<FCreatePlayerSanctionAppeal>(MoveTemp(Params));
+	if (!Op->IsReady())
+	{
+		Op->Then([this](TOnlineAsyncOp<FCreatePlayerSanctionAppeal>& InAsyncOp)
+		{
+			const FCreatePlayerSanctionAppeal::Params& Params = InAsyncOp.GetParams();
+
+			if (!Services.Get<FAuthEOSGS>()->IsLoggedIn(Params.LocalAccountId))
+			{
+					InAsyncOp.SetError(Errors::InvalidUser());
+			}
+		})
+		.Then([this](TOnlineAsyncOp<FCreatePlayerSanctionAppeal>& InAsyncOp, TPromise<const EOS_Sanctions_CreatePlayerSanctionAppealCallbackInfo*>&& Promise)
+		{
+			const FCreatePlayerSanctionAppeal::Params& Params = InAsyncOp.GetParams();
+
+			EOS_Sanctions_CreatePlayerSanctionAppealOptions Options = {}; 
+			Options.ApiVersion = 1;
+			UE_EOS_CHECK_API_MISMATCH(EOS_SANCTIONS_CREATEPLAYERSANCTIONAPPEAL_API_LATEST, 1);
+			Options.LocalUserId = GetProductUserIdChecked(Params.LocalAccountId);
+			Options.Reason = ToEOSSanctionAppealReason(Params.Reason);
+			const auto ReferenceIdUTF8 = StringCast<UTF8CHAR>(*Params.ReferenceId);
+			Options.ReferenceId = (const char*)ReferenceIdUTF8.Get();
+
+			EOS_Async(EOS_Sanctions_CreatePlayerSanctionAppeal, PlayerSanctionsHandle, Options, MoveTemp(Promise));
+		})
+		.Then([this](TOnlineAsyncOp<FCreatePlayerSanctionAppeal>& InAsyncOp, const EOS_Sanctions_CreatePlayerSanctionAppealCallbackInfo* Data)
+		{
+			if (Data->ResultCode != EOS_EResult::EOS_Success)
+			{
+				UE_LOG(LogOnlineServices, Warning, TEXT("EOS_Sanctions_CreatePlayerSanctionAppeal result=[%s]"), *LexToString(Data->ResultCode));
+				InAsyncOp.SetError(Errors::FromEOSResult(Data->ResultCode));
+			}
+			else
+			{
+				UE_LOG(LogOnlineServices, Verbose, TEXT("EOS_Sanctions_CreatePlayerSanctionAppeal result=[%s]"), *LexToString(Data->ResultCode));
+				InAsyncOp.SetResult({});
+			}
+		}).Enqueue(GetSerialQueue());
+	}
+	return Op->GetHandle();
+}
+
+TOnlineAsyncOpHandle<FReadActivePlayerSanctions> FPlayerSanctionsEOSGS::ReadEntriesForUser(FReadActivePlayerSanctions::Params&& Params)
+{
+	TOnlineAsyncOpRef<FReadActivePlayerSanctions> Op = GetOp<FReadActivePlayerSanctions>(MoveTemp(Params));
+	if (!Op->IsReady())
+	{
+		Op->Then([this](TOnlineAsyncOp<FReadActivePlayerSanctions>& InAsyncOp)
+		{
+			const FReadActivePlayerSanctions::Params& Params = InAsyncOp.GetParams();
+
+			if (!Services.Get<FAuthEOSGS>()->IsLoggedIn(Params.LocalAccountId))
+			{
+				InAsyncOp.SetError(Errors::InvalidUser());
+			}
+		})
+		.Then([this](TOnlineAsyncOp<FReadActivePlayerSanctions>& InAsyncOp, TPromise<const EOS_Sanctions_QueryActivePlayerSanctionsCallbackInfo*>&& Promise)
+		{
+			const FReadActivePlayerSanctions::Params& Params = InAsyncOp.GetParams();
+
+			EOS_Sanctions_QueryActivePlayerSanctionsOptions Options = {}; 
+			Options.ApiVersion = 2;
+			UE_EOS_CHECK_API_MISMATCH(EOS_SANCTIONS_QUERYACTIVEPLAYERSANCTIONS_API_LATEST, 2);
+			Options.LocalUserId = GetProductUserIdChecked(Params.LocalAccountId);
+			Options.TargetUserId = GetProductUserIdChecked(Params.TargetAccountId);
+
+			EOS_Async(EOS_Sanctions_QueryActivePlayerSanctions, PlayerSanctionsHandle, Options, MoveTemp(Promise));
+		})
+		.Then([this](TOnlineAsyncOp<FReadActivePlayerSanctions>& InAsyncOp, const EOS_Sanctions_QueryActivePlayerSanctionsCallbackInfo* Data)
+		{
+			if (Data->ResultCode != EOS_EResult::EOS_Success)
+			{
+				UE_LOG(LogOnlineServices, Warning, TEXT("EOS_Sanctions_QueryActivePlayerSanctions result=[%s]"), *LexToString(Data->ResultCode));
+				InAsyncOp.SetError(Errors::FromEOSResult(Data->ResultCode));
+				return;
+			}
+
+			UE_LOG(LogOnlineServices, Verbose, TEXT("EOS_Sanctions_QueryActivePlayerSanctions result=[%s]"), *LexToString(Data->ResultCode));
+
+			const FReadActivePlayerSanctions::Params& Params = InAsyncOp.GetParams();
+
+			EOS_Sanctions_GetPlayerSanctionCountOptions CountOptions = {};
+			CountOptions.ApiVersion = 1;
+			UE_EOS_CHECK_API_MISMATCH(EOS_SANCTIONS_GETPLAYERSANCTIONCOUNT_API_LATEST, 1);
+			CountOptions.TargetUserId = GetProductUserIdChecked(Params.TargetAccountId);
+
+			uint32 SanctionCount = EOS_Sanctions_GetPlayerSanctionCount(PlayerSanctionsHandle, &CountOptions);
+
+			EOS_Sanctions_CopyPlayerSanctionByIndexOptions CopyOptions = {};
+			CopyOptions.ApiVersion = 1;
+			UE_EOS_CHECK_API_MISMATCH(EOS_SANCTIONS_COPYPLAYERSANCTIONBYINDEX_API_LATEST, 1);
+			CopyOptions.TargetUserId = GetProductUserIdChecked(Params.TargetAccountId);
+
+			FReadActivePlayerSanctions::Result Result; 
+
+			for (uint32 SanctionIndex = 0; SanctionIndex < SanctionCount; SanctionIndex++)
+			{
+				CopyOptions.SanctionIndex = SanctionIndex;
+
+				EOS_Sanctions_PlayerSanction* PlayerSanctionEOS = nullptr;
+
+				EOS_EResult CopyResult = EOS_Sanctions_CopyPlayerSanctionByIndex(PlayerSanctionsHandle, &CopyOptions, &PlayerSanctionEOS);
+				if (CopyResult == EOS_EResult::EOS_Success)
+				{
+					FActivePlayerSanctionEntry& Entry = Result.Entries.Emplace_GetRef(); 
+					Entry.TimePlaced = PlayerSanctionEOS->TimePlaced;
+					Entry.TimeExpires = PlayerSanctionEOS->TimeExpires;
+					Entry.ReferenceId = PlayerSanctionEOS->ReferenceId;
+					Entry.Action = PlayerSanctionEOS->Action;
+
+					EOS_Sanctions_PlayerSanction_Release(PlayerSanctionEOS);
+				}
+				else
+				{
+					UE_LOG(LogOnlineServices, Warning, TEXT("EOS_Sanctions_CopyPlayerSanctionByIndex result=[%s]"), *LexToString(Data->ResultCode));
+				}
+			}
+
+			ToLogString(Result); 
+									
+			InAsyncOp.SetResult(MoveTemp(Result));
+		})
+		.Enqueue(GetSerialQueue());
+	}
+	return Op->GetHandle();
+}
+
+FString ToLogString(const FCreatePlayerSanctionAppeal::Result& CreatePlayerSanctionAppealResult)
+{
+	// CreatePlayerSanctionAppeal result is empty
+	return FString();
+}
+
+FString ToLogString(const FReadActivePlayerSanctions::Result& ReadPlayerSanctionResult)
+{
+	TArray<FString> StringEntries;
+
+	for (FActivePlayerSanctionEntry Entry : ReadPlayerSanctionResult.Entries)
+	{
+		StringEntries.Emplace(FString::Printf(TEXT("ReferenceId:%s Action:%s TimeExpires:%llu TimePlaced:%llu"), *Entry.ReferenceId, *Entry.Action, Entry.TimeExpires, Entry.TimePlaced));
+	}
+
+	return FString::Join(StringEntries, TEXT(","));
+}
+/* UE::Online */ }
\ No newline at end of file
diff --git a/Engine/Plugins/Online/OnlineServicesEOSGS/Source/Public/Online/OnlineServicesEOSGS.h b/Engine/Plugins/Online/OnlineServicesEOSGS/Source/Public/Online/OnlineServicesEOSGS.h
--- a/Engine/Plugins/Online/OnlineServicesEOSGS/Source/Public/Online/OnlineServicesEOSGS.h
+++ b/Engine/Plugins/Online/OnlineServicesEOSGS/Source/Public/Online/OnlineServicesEOSGS.h
@@ -17,6 +17,9 @@ using IEOSPlatformHandlePtr = TSharedPtr<class IEOSPlatformHandle>;
 
 namespace UE::Online {
 
+using IPlayerReportsPtr = TSharedPtr<class IPlayerReports>;
+using IPlayerSanctionsPtr = TSharedPtr<class IPlayerSanctions>;
+
 class ONLINESERVICESEOSGS_API FOnlineServicesEOSGS : public FOnlineServicesEpicCommon
 {
 public:
@@ -31,6 +34,9 @@ public:
 	virtual TOnlineResult<FGetResolvedConnectString> GetResolvedConnectString(FGetResolvedConnectString::Params&& Params) override;
 	virtual EOnlineServices GetServicesProvider() const override { return EOnlineServices::Epic; }
 
+	IPlayerReportsPtr GetPlayerReportsInterface();
+	IPlayerSanctionsPtr GetPlayerSanctionsInterface(); 
+
 	UE_DEPRECATED(5.5, "Please call GetServiceConfigNameStatic instead")
 	static const TCHAR* GetConfigNameStatic() { return GetServiceConfigNameStatic(); }
 	static const TCHAR* GetServiceConfigNameStatic() { return TEXT("EOS"); }
diff --git a/Engine/Plugins/Online/OnlineServicesEOSGS/Source/Public/Online/OnlineServicesEOSGSInterfaces/PlayerReports.h b/Engine/Plugins/Online/OnlineServicesEOSGS/Source/Public/Online/OnlineServicesEOSGSInterfaces/PlayerReports.h
new file mode 100644
--- /dev/null
+++ b/Engine/Plugins/Online/OnlineServicesEOSGS/Source/Public/Online/OnlineServicesEOSGSInterfaces/PlayerReports.h
@@ -0,0 +1,78 @@
+// Copyright Epic Games, Inc. All Rights Reserved.
+
+#pragma once
+
+#include "Online/CoreOnline.h"
+#include "Online/OnlineMeta.h"
+
+namespace UE::Online { template <typename OpType> class TOnlineAsyncOpHandle; }
+namespace UE::Online { template <typename OpType> class TOnlineResult; }
+namespace UE::Online::Meta { template <typename StructType> struct TStructDetails; }
+
+namespace UE::Online {
+
+enum class EPlayerReportCategory
+{
+	Cheating,
+	Exploiting,
+	OffensiveProfile,
+	VerbalAbuse,
+	Scamming,
+	Spamming,
+	Other
+};
+
+const TCHAR* LexToString(EPlayerReportCategory Value);
+void LexFromString(EPlayerReportCategory& OutValue, const TCHAR* InStr);
+
+struct FSendPlayerReport
+{
+	static constexpr TCHAR Name[] = TEXT("SendPlayerReport");
+
+	struct Params
+	{
+		/* Required - Local user performing the operation */
+		FAccountId LocalAccountId;
+		/* Required - Target user of the player report */
+		FAccountId TargetAccountId;
+		/* Required - The category of the player */
+		EPlayerReportCategory Category;
+		/* Optional - Message from the player describing the issue being reported */
+		FString Message = "";
+		/* Optional - Context information around the issue being reported the game can send. This needs to be in a valid JSON format otherwise the EOS SDK will throw an error */
+		FString Context = "";
+	};
+
+	struct Result
+	{
+	};
+};
+
+/**
+ * Interface definition for the EOS player reports service
+ */
+class IPlayerReports
+{
+public:
+	/**
+	* Send a player report for a player cheating, misconduct, etc... 
+	*/
+	virtual TOnlineAsyncOpHandle<FSendPlayerReport> SendPlayerReport(FSendPlayerReport::Params&& Params) = 0;
+};
+
+namespace Meta {
+
+BEGIN_ONLINE_STRUCT_META(FSendPlayerReport::Params)
+	ONLINE_STRUCT_FIELD(FSendPlayerReport::Params, LocalAccountId),
+	ONLINE_STRUCT_FIELD(FSendPlayerReport::Params, TargetAccountId),
+	ONLINE_STRUCT_FIELD(FSendPlayerReport::Params, Category), 
+	ONLINE_STRUCT_FIELD(FSendPlayerReport::Params, Message), 
+	ONLINE_STRUCT_FIELD(FSendPlayerReport::Params, Context)
+END_ONLINE_STRUCT_META()
+
+BEGIN_ONLINE_STRUCT_META(FSendPlayerReport::Result)
+END_ONLINE_STRUCT_META()
+
+/* Meta*/ }
+
+/* UE::Online */ }
diff --git a/Engine/Plugins/Online/OnlineServicesEOSGS/Source/Public/Online/OnlineServicesEOSGSInterfaces/PlayerSanctions.h b/Engine/Plugins/Online/OnlineServicesEOSGS/Source/Public/Online/OnlineServicesEOSGSInterfaces/PlayerSanctions.h
new file mode 100644
--- /dev/null
+++ b/Engine/Plugins/Online/OnlineServicesEOSGS/Source/Public/Online/OnlineServicesEOSGSInterfaces/PlayerSanctions.h
@@ -0,0 +1,114 @@
+// Copyright Epic Games, Inc. All Rights Reserved.
+
+#pragma once
+
+#include "Online/CoreOnline.h"
+#include "Online/OnlineMeta.h"
+
+namespace UE::Online { template <typename OpType> class TOnlineAsyncOpHandle; }
+namespace UE::Online { template <typename OpType> class TOnlineResult; }
+namespace UE::Online::Meta { template <typename StructType> struct TStructDetails; }
+
+namespace UE::Online {
+
+enum class EPlayerSanctionAppealReason
+{
+	IncorrectSanction,
+	CompromisedAccount,
+	UnfairPunishment,
+	AppealForForgiveness
+};
+
+const TCHAR* LexToString(EPlayerSanctionAppealReason Value);
+void LexFromString(EPlayerSanctionAppealReason& OutValue, const TCHAR* InStr);
+
+struct FCreatePlayerSanctionAppeal
+{
+	static constexpr TCHAR Name[] = TEXT("CreatePlayerSanctionAppeal");
+
+	struct Params
+	{
+		/* Local user id */
+		FAccountId LocalAccountId;
+		/* The reason of the appeal */
+		EPlayerSanctionAppealReason Reason;
+		/* The sanction id for the sanction that is being appealed */
+		FString ReferenceId;
+	};
+
+	struct Result
+	{
+	};
+};
+
+struct FActivePlayerSanctionEntry
+{
+	/* The time the sanction was placed */
+	int64_t TimePlaced;
+	/* The time the sanction expires */
+	int64_t TimeExpires;
+	/* The action associated with this sanction */
+	FString Action;
+	/* The sanction id for the sanction that is being appealed. This needs to be set */
+	FString ReferenceId;
+};
+
+struct FReadActivePlayerSanctions
+{
+	static constexpr TCHAR Name[] = TEXT("ReadActivePlayerSanctions");
+
+	struct Params
+	{
+		/* Local user id */
+		FAccountId LocalAccountId;
+		/* The account ids of the user we are querying active sanctions for */
+		FAccountId TargetAccountId;
+	};
+
+	struct Result
+	{
+		/* The result leaderboard entries */
+		TArray<FActivePlayerSanctionEntry> Entries;
+	};
+};
+
+/**
+ * Interface definition for the EOS player reports service
+ */
+class IPlayerSanctions
+{
+public:
+	/**
+	* Send a player sanction appeal
+	*/
+	virtual TOnlineAsyncOpHandle<FCreatePlayerSanctionAppeal> CreatePlayerSanctionAppeal(FCreatePlayerSanctionAppeal::Params&& Params) = 0;
+
+	/**
+	 * Read active player sanction for a specific user
+	 */
+	virtual TOnlineAsyncOpHandle<FReadActivePlayerSanctions> ReadEntriesForUser(FReadActivePlayerSanctions::Params&& Params) = 0;
+};
+
+namespace Meta {
+
+BEGIN_ONLINE_STRUCT_META(FCreatePlayerSanctionAppeal::Params)
+	ONLINE_STRUCT_FIELD(FCreatePlayerSanctionAppeal::Params, LocalAccountId),
+	ONLINE_STRUCT_FIELD(FCreatePlayerSanctionAppeal::Params, Reason),
+	ONLINE_STRUCT_FIELD(FCreatePlayerSanctionAppeal::Params, ReferenceId)
+END_ONLINE_STRUCT_META()
+
+BEGIN_ONLINE_STRUCT_META(FReadActivePlayerSanctions::Params)
+	ONLINE_STRUCT_FIELD(FReadActivePlayerSanctions::Params, LocalAccountId),
+	ONLINE_STRUCT_FIELD(FReadActivePlayerSanctions::Params, TargetAccountId)
+END_ONLINE_STRUCT_META()
+
+BEGIN_ONLINE_STRUCT_META(FCreatePlayerSanctionAppeal::Result)
+END_ONLINE_STRUCT_META()
+
+BEGIN_ONLINE_STRUCT_META(FReadActivePlayerSanctions::Result)
+	ONLINE_STRUCT_FIELD(FReadActivePlayerSanctions::Result, Entries)
+END_ONLINE_STRUCT_META()
+
+/* Meta*/ }
+
+/* UE::Online */ }
diff --git a/Engine/Plugins/Online/OnlineServicesEOSGS/Source/Public/Online/PlayerReportsEOSGS.h b/Engine/Plugins/Online/OnlineServicesEOSGS/Source/Public/Online/PlayerReportsEOSGS.h
new file mode 100644
--- /dev/null
+++ b/Engine/Plugins/Online/OnlineServicesEOSGS/Source/Public/Online/PlayerReportsEOSGS.h
@@ -0,0 +1,44 @@
+// Copyright Epic Games, Inc. All Rights Reserved.
+
+#pragma once
+
+#include "CoreMinimal.h"
+#include "Online/OnlineComponent.h"
+#include "OnlineServicesEOSGSInterfaces/PlayerReports.h"
+
+#if defined(EOS_PLATFORM_BASE_FILE_NAME)
+#include EOS_PLATFORM_BASE_FILE_NAME
+#endif
+#include "eos_reports_types.h"
+
+enum class EPlayerReportCategory : uint8;
+namespace UE::Online { template <typename Result> class TDefaultErrorResultInternal; }
+template <typename ResultType> class TFuture;
+
+namespace UE::Online {
+
+class FOnlineServicesEOSGS;
+
+class ONLINESERVICESEOSGS_API FPlayerReportsEOSGS : public TOnlineComponent<IPlayerReports>
+{
+public:
+	using Super = TOnlineComponent<IPlayerReports>;
+
+	FPlayerReportsEOSGS(FOnlineServicesEOSGS& InOwningSubsystem);
+	virtual ~FPlayerReportsEOSGS() = default;
+
+	// IOnlineComponent
+	virtual void Initialize() override;
+	virtual void RegisterCommands() override;
+
+	// IPlayerReports	
+	virtual TOnlineAsyncOpHandle<FSendPlayerReport> SendPlayerReport(FSendPlayerReport::Params&& Params) override; 
+
+private:
+	EOS_HReports PlayerReportsHandle = nullptr;
+};
+
+FString ToLogString(const FSendPlayerReport::Result& PlayerReports);
+/* UE::Online */ }
+
+
diff --git a/Engine/Plugins/Online/OnlineServicesEOSGS/Source/Public/Online/PlayerSanctionsEOSGS.h b/Engine/Plugins/Online/OnlineServicesEOSGS/Source/Public/Online/PlayerSanctionsEOSGS.h
new file mode 100644
--- /dev/null
+++ b/Engine/Plugins/Online/OnlineServicesEOSGS/Source/Public/Online/PlayerSanctionsEOSGS.h
@@ -0,0 +1,45 @@
+// Copyright Epic Games, Inc. All Rights Reserved.
+
+#pragma once
+
+#include "CoreMinimal.h"
+#include "Online/OnlineComponent.h"
+#include "OnlineServicesEOSGSInterfaces/PlayerSanctions.h"
+
+
+#if defined(EOS_PLATFORM_BASE_FILE_NAME)
+#include EOS_PLATFORM_BASE_FILE_NAME
+#endif
+#include "eos_sanctions_types.h"
+
+
+namespace UE::Online {
+
+class FOnlineServicesEOSGS;
+
+class ONLINESERVICESEOSGS_API FPlayerSanctionsEOSGS : public TOnlineComponent<IPlayerSanctions>
+{
+public:
+	using Super = TOnlineComponent<IPlayerSanctions>;
+
+	FPlayerSanctionsEOSGS(FOnlineServicesEOSGS& InOwningSubsystem);
+	virtual ~FPlayerSanctionsEOSGS() = default;
+
+	// IOnlineComponent
+	virtual void Initialize() override;
+	virtual void RegisterCommands() override;
+
+	// IPlayerSanctions
+	virtual TOnlineAsyncOpHandle<FCreatePlayerSanctionAppeal> CreatePlayerSanctionAppeal(FCreatePlayerSanctionAppeal::Params&& Params) override;
+	virtual TOnlineAsyncOpHandle<FReadActivePlayerSanctions> ReadEntriesForUser(FReadActivePlayerSanctions::Params&& Params) override;
+
+private:
+	EOS_HSanctions PlayerSanctionsHandle = nullptr;
+};
+
+FString ToLogString(const FReadActivePlayerSanctions::Result& ReadPlayerSanctionResult);
+FString ToLogString(const FReadActivePlayerSanctions::Result& ReadPlayerSanctionResult);
+
+/* UE::Online */ }
+
+
