- Register more EOS[GS] components in the split plugins

#rnx

diff --git a/Engine/Plugins/Online/OnlineServicesEOS/Source/Private/Online/ExternalUIEOS.cpp b/Engine/Plugins/Online/OnlineServicesEOS/Source/Private/Online/ExternalUIEOS.cpp
--- a/Engine/Plugins/Online/OnlineServicesEOS/Source/Private/Online/ExternalUIEOS.cpp
+++ b/Engine/Plugins/Online/OnlineServicesEOS/Source/Private/Online/ExternalUIEOS.cpp
@@ -12,7 +12,7 @@
 
 namespace UE::Online {
 
-FExternalUIEOS::FExternalUIEOS(FOnlineServicesEOS& InServices)
+FExternalUIEOS::FExternalUIEOS(FOnlineServicesEpicCommon& InServices)
 	: Super(InServices)
 {
 }
@@ -21,7 +21,7 @@ void FExternalUIEOS::Initialize()
 {
 	Super::Initialize();
 
-	UIHandle = EOS_Platform_GetUIInterface(*GetServices<FOnlineServicesEOS>().GetEOSPlatformHandle());
+	UIHandle = EOS_Platform_GetUIInterface(*GetServices<FOnlineServicesEpicCommon>().GetEOSPlatformHandle());
 	check(UIHandle != nullptr);
 }
 
diff --git a/Engine/Plugins/Online/OnlineServicesEOS/Source/Private/Online/PresenceEOS.cpp b/Engine/Plugins/Online/OnlineServicesEOS/Source/Private/Online/PresenceEOS.cpp
--- a/Engine/Plugins/Online/OnlineServicesEOS/Source/Private/Online/PresenceEOS.cpp
+++ b/Engine/Plugins/Online/OnlineServicesEOS/Source/Private/Online/PresenceEOS.cpp
@@ -61,7 +61,7 @@ static inline EOS_Presence_EStatus ToEOS_Presence_EStatus(EUserPresenceStatus In
 	return EOS_Presence_EStatus::EOS_PS_Offline;
 }
 
-FPresenceEOS::FPresenceEOS(FOnlineServicesEOS& InServices)
+FPresenceEOS::FPresenceEOS(FOnlineServicesEpicCommon& InServices)
 	: Super(InServices)
 {
 }
@@ -70,7 +70,7 @@ void FPresenceEOS::Initialize()
 {
 	Super::Initialize();
 
-	PresenceHandle = EOS_Platform_GetPresenceInterface(*GetServices<FOnlineServicesEOS>().GetEOSPlatformHandle());
+	PresenceHandle = EOS_Platform_GetPresenceInterface(*GetServices<FOnlineServicesEpicCommon>().GetEOSPlatformHandle());
 	check(PresenceHandle != nullptr);
 
 	LoginStatusChangedHandle = Services.Get<IAuth>()->OnLoginStatusChanged().Add(this, &FPresenceEOS::HandleAuthLoginStatusChanged);
diff --git a/Engine/Plugins/Online/OnlineServicesEOS/Source/Private/Online/SocialEOS.cpp b/Engine/Plugins/Online/OnlineServicesEOS/Source/Private/Online/SocialEOS.cpp
--- a/Engine/Plugins/Online/OnlineServicesEOS/Source/Private/Online/SocialEOS.cpp
+++ b/Engine/Plugins/Online/OnlineServicesEOS/Source/Private/Online/SocialEOS.cpp
@@ -26,7 +26,7 @@ ERelationship EOSEFriendStatusToERelationship(EOS_EFriendsStatus Status)
 	return ERelationship::NotFriend;
 }
 
-FSocialEOS::FSocialEOS(FOnlineServicesEOS& InServices)
+FSocialEOS::FSocialEOS(FOnlineServicesEpicCommon& InServices)
 	: FSocialCommon(InServices)
 {
 }
@@ -35,10 +35,10 @@ void FSocialEOS::Initialize()
 {
 	FSocialCommon::Initialize();
 
-	FriendsHandle = EOS_Platform_GetFriendsInterface(*GetServices<FOnlineServicesEOS>().GetEOSPlatformHandle());
+	FriendsHandle = EOS_Platform_GetFriendsInterface(*GetServices<FOnlineServicesEpicCommon>().GetEOSPlatformHandle());
 	check(FriendsHandle != nullptr);
 
-	UIHandle = EOS_Platform_GetUIInterface(*GetServices<FOnlineServicesEOS>().GetEOSPlatformHandle());
+	UIHandle = EOS_Platform_GetUIInterface(*GetServices<FOnlineServicesEpicCommon>().GetEOSPlatformHandle());
 
 	OnFriendsUpdate = EOS_RegisterComponentEventHandler(
 		this,
diff --git a/Engine/Plugins/Online/OnlineServicesEOS/Source/Private/Online/UserInfoEOS.cpp b/Engine/Plugins/Online/OnlineServicesEOS/Source/Private/Online/UserInfoEOS.cpp
--- a/Engine/Plugins/Online/OnlineServicesEOS/Source/Private/Online/UserInfoEOS.cpp
+++ b/Engine/Plugins/Online/OnlineServicesEOS/Source/Private/Online/UserInfoEOS.cpp
@@ -13,7 +13,7 @@
 
 namespace UE::Online {
 
-FUserInfoEOS::FUserInfoEOS(FOnlineServicesEOS& InServices)
+FUserInfoEOS::FUserInfoEOS(FOnlineServicesEpicCommon& InServices)
 	: FUserInfoCommon(InServices)
 {
 }
@@ -22,7 +22,7 @@ void FUserInfoEOS::Initialize()
 {
 	Super::Initialize();
 
-	UserInfoHandle = EOS_Platform_GetUserInfoInterface(*GetServices<FOnlineServicesEOS>().GetEOSPlatformHandle());
+	UserInfoHandle = EOS_Platform_GetUserInfoInterface(*GetServices<FOnlineServicesEpicCommon>().GetEOSPlatformHandle());
 	check(UserInfoHandle != nullptr);
 }
 
diff --git a/Engine/Plugins/Online/OnlineServicesEOS/Source/Public/Online/ExternalUIEOS.h b/Engine/Plugins/Online/OnlineServicesEOS/Source/Public/Online/ExternalUIEOS.h
--- a/Engine/Plugins/Online/OnlineServicesEOS/Source/Public/Online/ExternalUIEOS.h
+++ b/Engine/Plugins/Online/OnlineServicesEOS/Source/Public/Online/ExternalUIEOS.h
@@ -2,7 +2,6 @@
 
 #pragma once
 
-#include "CoreMinimal.h"
 #include "Online/ExternalUIEOSGS.h"
 
 #if defined(EOS_PLATFORM_BASE_FILE_NAME)
@@ -18,14 +17,14 @@ struct FExternalUIEOSConfig
 	bool bShowLoginUIEnabled = true;
 };
 
-class FOnlineServicesEOS;
+class FOnlineServicesEpicCommon;
 
 class ONLINESERVICESEOS_API FExternalUIEOS : public FExternalUIEOSGS
 {
 public:
 	using Super = FExternalUIEOSGS;
 
-	FExternalUIEOS(FOnlineServicesEOS& InOwningSubsystem);
+	FExternalUIEOS(FOnlineServicesEpicCommon& InOwningSubsystem);
 	virtual void Initialize() override;
 	virtual void UpdateConfig() override;
 	virtual TOnlineAsyncOpHandle<FExternalUIShowLoginUI> ShowLoginUI(FExternalUIShowLoginUI::Params&& Params) override;
diff --git a/Engine/Plugins/Online/OnlineServicesEOS/Source/Public/Online/PresenceEOS.h b/Engine/Plugins/Online/OnlineServicesEOS/Source/Public/Online/PresenceEOS.h
--- a/Engine/Plugins/Online/OnlineServicesEOS/Source/Public/Online/PresenceEOS.h
+++ b/Engine/Plugins/Online/OnlineServicesEOS/Source/Public/Online/PresenceEOS.h
@@ -13,14 +13,14 @@
 
 namespace UE::Online {
 
-class FOnlineServicesEOS;
+class FOnlineServicesEpicCommon;
 	
 class FPresenceEOS : public FPresenceCommon
 {
 public:
 	using Super = FPresenceCommon;
 
-	FPresenceEOS(FOnlineServicesEOS& InServices);
+	FPresenceEOS(FOnlineServicesEpicCommon& InServices);
 
 	virtual void Initialize() override;
 	virtual void PreShutdown() override;
diff --git a/Engine/Plugins/Online/OnlineServicesEOS/Source/Public/Online/SocialEOS.h b/Engine/Plugins/Online/OnlineServicesEOS/Source/Public/Online/SocialEOS.h
--- a/Engine/Plugins/Online/OnlineServicesEOS/Source/Public/Online/SocialEOS.h
+++ b/Engine/Plugins/Online/OnlineServicesEOS/Source/Public/Online/SocialEOS.h
@@ -2,7 +2,6 @@
 
 #pragma once
 
-#include "CoreMinimal.h"
 #include "Online/OnlineServicesEOSGSTypes.h"
 #include "Online/SocialCommon.h"
 
@@ -14,14 +13,14 @@
 
 namespace UE::Online {
 
-class FOnlineServicesEOS;
+class FOnlineServicesEpicCommon;
 	
 class FSocialEOS : public FSocialCommon
 {
 public:
 	using Super = FSocialCommon;
 
-	FSocialEOS(FOnlineServicesEOS& InServices);
+	FSocialEOS(FOnlineServicesEpicCommon& InServices);
 
 	virtual void Initialize() override;
 	virtual void PreShutdown() override;
diff --git a/Engine/Plugins/Online/OnlineServicesEOS/Source/Public/Online/UserInfoEOS.h b/Engine/Plugins/Online/OnlineServicesEOS/Source/Public/Online/UserInfoEOS.h
--- a/Engine/Plugins/Online/OnlineServicesEOS/Source/Public/Online/UserInfoEOS.h
+++ b/Engine/Plugins/Online/OnlineServicesEOS/Source/Public/Online/UserInfoEOS.h
@@ -2,7 +2,6 @@
 
 #pragma once
 
-#include "CoreMinimal.h"
 #include "Online/UserInfoCommon.h"
 
 #if defined(EOS_PLATFORM_BASE_FILE_NAME)
@@ -12,14 +11,14 @@
 
 namespace UE::Online {
 
-class FOnlineServicesEOS;
+class FOnlineServicesEpicCommon;
 	
 class FUserInfoEOS : public FUserInfoCommon
 {
 public:
 	using Super = FUserInfoCommon;
 
-	FUserInfoEOS(FOnlineServicesEOS& InServices);
+	FUserInfoEOS(FOnlineServicesEpicCommon& InServices);
 	virtual ~FUserInfoEOS() = default;
 
 	// IOnlineComponent
diff --git a/Engine/Plugins/Online/OnlineServicesEOSGS/Source/Private/Online/AchievementsEOSGS.cpp b/Engine/Plugins/Online/OnlineServicesEOSGS/Source/Private/Online/AchievementsEOSGS.cpp
--- a/Engine/Plugins/Online/OnlineServicesEOSGS/Source/Private/Online/AchievementsEOSGS.cpp
+++ b/Engine/Plugins/Online/OnlineServicesEOSGS/Source/Private/Online/AchievementsEOSGS.cpp
@@ -14,7 +14,7 @@
 
 namespace UE::Online {
 
-FAchievementsEOSGS::FAchievementsEOSGS(FOnlineServicesEOSGS& InServices)
+	FAchievementsEOSGS::FAchievementsEOSGS(FOnlineServicesEpicCommon& InServices)
 	: Super(InServices)
 {
 }
@@ -23,7 +23,7 @@ void FAchievementsEOSGS::Initialize()
 {
 	Super::Initialize();
 
-	AchievementsHandle = EOS_Platform_GetAchievementsInterface(*GetServices<FOnlineServicesEOSGS>().GetEOSPlatformHandle());
+	AchievementsHandle = EOS_Platform_GetAchievementsInterface(*GetServices<FOnlineServicesEpicCommon>().GetEOSPlatformHandle());
 	check(AchievementsHandle);
 
 	OnAchievementsUnlocked = EOS_RegisterComponentEventHandler(
diff --git a/Engine/Plugins/Online/OnlineServicesEOSGS/Source/Private/Online/ExternalUIEOSGS.cpp b/Engine/Plugins/Online/OnlineServicesEOSGS/Source/Private/Online/ExternalUIEOSGS.cpp
--- a/Engine/Plugins/Online/OnlineServicesEOSGS/Source/Private/Online/ExternalUIEOSGS.cpp
+++ b/Engine/Plugins/Online/OnlineServicesEOSGS/Source/Private/Online/ExternalUIEOSGS.cpp
@@ -9,7 +9,7 @@
 
 namespace UE::Online {
 
-FExternalUIEOSGS::FExternalUIEOSGS(FOnlineServicesEOSGS& InServices)
+FExternalUIEOSGS::FExternalUIEOSGS(FOnlineServicesEpicCommon& InServices)
 	: Super(InServices)
 {
 }
@@ -21,7 +21,7 @@ void FExternalUIEOSGS::Initialize()
 	IEOSSDKManager* SDKManager = IEOSSDKManager::Get();
 	check(SDKManager);
 
-	UIInterfaceHandle = EOS_Platform_GetUIInterface(*GetServices<FOnlineServicesEOSGS>().GetEOSPlatformHandle());
+	UIInterfaceHandle = EOS_Platform_GetUIInterface(*GetServices<FOnlineServicesEpicCommon>().GetEOSPlatformHandle());
 	check(UIInterfaceHandle != nullptr);
 
 	RegisterEventHandlers();
diff --git a/Engine/Plugins/Online/OnlineServicesEOSGS/Source/Private/Online/LeaderboardsEOSGS.cpp b/Engine/Plugins/Online/OnlineServicesEOSGS/Source/Private/Online/LeaderboardsEOSGS.cpp
--- a/Engine/Plugins/Online/OnlineServicesEOSGS/Source/Private/Online/LeaderboardsEOSGS.cpp
+++ b/Engine/Plugins/Online/OnlineServicesEOSGS/Source/Private/Online/LeaderboardsEOSGS.cpp
@@ -16,7 +16,7 @@ namespace UE::Online {
 #define EOS_LEADERBOARD_STAT_NAME_MAX_LENGTH 256 + 1 // 256 plus null terminator
 #define EOS_LEADERBOARD_MAX_ENTRIES 1000
 
-FLeaderboardsEOSGS::FLeaderboardsEOSGS(FOnlineServicesEOSGS& InServices)
+FLeaderboardsEOSGS::FLeaderboardsEOSGS(FOnlineServicesEpicCommon& InServices)
 	: Super(InServices)
 {
 }
@@ -25,7 +25,7 @@ void FLeaderboardsEOSGS::Initialize()
 {
 	Super::Initialize();
 
-	LeaderboardsHandle = EOS_Platform_GetLeaderboardsInterface(*GetServices<FOnlineServicesEOSGS>().GetEOSPlatformHandle());
+	LeaderboardsHandle = EOS_Platform_GetLeaderboardsInterface(*GetServices<FOnlineServicesEpicCommon>().GetEOSPlatformHandle());
 	check(LeaderboardsHandle);
 }
 
diff --git a/Engine/Plugins/Online/OnlineServicesEOSGS/Source/Private/Online/LobbiesEOSGS.cpp b/Engine/Plugins/Online/OnlineServicesEOSGS/Source/Private/Online/LobbiesEOSGS.cpp
--- a/Engine/Plugins/Online/OnlineServicesEOSGS/Source/Private/Online/LobbiesEOSGS.cpp
+++ b/Engine/Plugins/Online/OnlineServicesEOSGS/Source/Private/Online/LobbiesEOSGS.cpp
@@ -24,7 +24,7 @@ namespace UE::Online {
 
 static const int32 MaxAttributeSize = 1000;
 
-FLobbiesEOSGS::FLobbiesEOSGS(FOnlineServicesEOSGS& InServices)
+FLobbiesEOSGS::FLobbiesEOSGS(FOnlineServicesEpicCommon& InServices)
 	: FLobbiesCommon(InServices)
 {
 }
@@ -36,7 +36,7 @@ void FLobbiesEOSGS::Initialize()
 	IEOSSDKManager* SDKManager = IEOSSDKManager::Get();
 	check(SDKManager);
 
-	EOS_HLobby LobbyInterfaceHandle = EOS_Platform_GetLobbyInterface(*GetServices<FOnlineServicesEOSGS>().GetEOSPlatformHandle());
+	EOS_HLobby LobbyInterfaceHandle = EOS_Platform_GetLobbyInterface(*GetServices<FOnlineServicesEpicCommon>().GetEOSPlatformHandle());
 	check(LobbyInterfaceHandle != nullptr);
 
 	LobbyPrerequisites = MakeShared<FLobbyPrerequisitesEOS>(FLobbyPrerequisitesEOS{
diff --git a/Engine/Plugins/Online/OnlineServicesEOSGS/Source/Private/Online/OnlineServicesEOSGS.cpp b/Engine/Plugins/Online/OnlineServicesEOSGS/Source/Private/Online/OnlineServicesEOSGS.cpp
--- a/Engine/Plugins/Online/OnlineServicesEOSGS/Source/Private/Online/OnlineServicesEOSGS.cpp
+++ b/Engine/Plugins/Online/OnlineServicesEOSGS/Source/Private/Online/OnlineServicesEOSGS.cpp
@@ -65,27 +65,6 @@ void FOnlineServicesEOSGS::Destroy()
 #endif
 }
 
-void FOnlineServicesEOSGS::WarnIfEncryptionKeyMissing(const FString& InterfaceName) const
-{
-	if (IEOSSDKManager* Manager = IEOSSDKManager::Get())
-	{
-		const FString& PlatformConfigName = GetEOSPlatformHandle()->GetConfigName();
-		if (const FEOSSDKPlatformConfig* Config = Manager->GetPlatformConfig(PlatformConfigName))
-		{
-			const FString& EncryptionKey = Config->EncryptionKey;
-			if (EncryptionKey.IsEmpty())
-			{
-				UE_LOG(LogOnlineServices, Verbose, TEXT("%s interface not available due to missing ClientEncryptionKey in config."), *InterfaceName);
-			}
-			else
-			{
-				// If we have an encryption key and still can't get the interface, something weird is going on.
-				UE_LOG(LogOnlineServices, Warning, TEXT("%s interface not available despite encryption key being present."), *InterfaceName);
-			}
-		}
-	}
-}
-
 void FOnlineServicesEOSGS::RegisterComponents()
 {
 	Components.Register<FAchievementsEOSGS>(*this);
diff --git a/Engine/Plugins/Online/OnlineServicesEOSGS/Source/Private/Online/OnlineServicesEOSGSModule.cpp b/Engine/Plugins/Online/OnlineServicesEOSGS/Source/Private/Online/OnlineServicesEOSGSModule.cpp
--- a/Engine/Plugins/Online/OnlineServicesEOSGS/Source/Private/Online/OnlineServicesEOSGSModule.cpp
+++ b/Engine/Plugins/Online/OnlineServicesEOSGS/Source/Private/Online/OnlineServicesEOSGSModule.cpp
@@ -52,8 +52,12 @@ void FOnlineServicesEOSGSModule::StartupModule()
 
 	static FOnlineAccountIdRegistryEOSGS AccountIdRegistry(EOnlineServices::Epic);
 	FOnlineIdRegistryRegistry::Get().RegisterAccountIdRegistry(EOnlineServices::Epic, &AccountIdRegistry, GetRegistryPriority());
-	FOnlineIdRegistryRegistry::Get().RegisterSessionIdRegistry(EOnlineServices::Epic, &FOnlineSessionIdRegistryEOSGS::Get(), GetRegistryPriority());
-	FOnlineIdRegistryRegistry::Get().RegisterSessionInviteIdRegistry(EOnlineServices::Epic, &FOnlineSessionInviteIdRegistryEOSGS::Get(), GetRegistryPriority());
+
+	static FOnlineSessionIdRegistryEOSGS SessionIdRegistry(EOnlineServices::Epic);
+	FOnlineIdRegistryRegistry::Get().RegisterSessionIdRegistry(EOnlineServices::Epic, &SessionIdRegistry, GetRegistryPriority());
+
+	static FOnlineSessionInviteIdRegistryEOSGS SessionInviteIdRegistry(EOnlineServices::Epic);
+	FOnlineIdRegistryRegistry::Get().RegisterSessionInviteIdRegistry(EOnlineServices::Epic, &SessionInviteIdRegistry, GetRegistryPriority());
 
 	// Initialize the platform factory on startup.  This is necessary for the SDK to bind to rendering and input very early.
 	FOnlineServicesEOSGSPlatformFactory::Get();
diff --git a/Engine/Plugins/Online/OnlineServicesEOSGS/Source/Private/Online/SessionsEOSGS.cpp b/Engine/Plugins/Online/OnlineServicesEOSGS/Source/Private/Online/SessionsEOSGS.cpp
--- a/Engine/Plugins/Online/OnlineServicesEOSGS/Source/Private/Online/SessionsEOSGS.cpp
+++ b/Engine/Plugins/Online/OnlineServicesEOSGS/Source/Private/Online/SessionsEOSGS.cpp
@@ -17,33 +17,35 @@ namespace UE::Online {
 
 /** FOnlineSessionIdRegistryEOSGS */
 
-FOnlineSessionIdRegistryEOSGS::FOnlineSessionIdRegistryEOSGS()
-	: FOnlineSessionIdRegistryLAN(EOnlineServices::Epic)
+FOnlineSessionIdRegistryEOSGS::FOnlineSessionIdRegistryEOSGS(EOnlineServices ServicesType)
+	: FOnlineSessionIdRegistryLAN(ServicesType)
 {
 }
 
-FOnlineSessionIdRegistryEOSGS& FOnlineSessionIdRegistryEOSGS::Get()
+FOnlineSessionIdRegistryEOSGS& FOnlineSessionIdRegistryEOSGS::GetRegistered(EOnlineServices Services)
 {
-	static FOnlineSessionIdRegistryEOSGS Instance;
-	return Instance;
+	IOnlineSessionIdRegistry* Registry = FOnlineIdRegistryRegistry::Get().GetSessionIdRegistry(Services);
+	check(Registry);
+	return *static_cast<FOnlineSessionIdRegistryEOSGS*>(Registry);
 }
 
 /** FOnlineSessionInviteIdRegistryEOSGS */
 
-FOnlineSessionInviteIdRegistryEOSGS::FOnlineSessionInviteIdRegistryEOSGS()
-	: FOnlineSessionInviteIdStringRegistry(EOnlineServices::Epic)
+FOnlineSessionInviteIdRegistryEOSGS::FOnlineSessionInviteIdRegistryEOSGS(EOnlineServices ServicesType)
+	: FOnlineSessionInviteIdStringRegistry(ServicesType)
 {
 }
 
-FOnlineSessionInviteIdRegistryEOSGS& FOnlineSessionInviteIdRegistryEOSGS::Get()
+FOnlineSessionInviteIdRegistryEOSGS& FOnlineSessionInviteIdRegistryEOSGS::GetRegistered(EOnlineServices Services)
 {
-	static FOnlineSessionInviteIdRegistryEOSGS Instance;
-	return Instance;
+	IOnlineSessionInviteIdRegistry* Registry = FOnlineIdRegistryRegistry::Get().GetSessionInviteIdRegistry(Services);
+	check(Registry);
+	return *static_cast<FOnlineSessionInviteIdRegistryEOSGS*>(Registry);
 }
 
 /** FSessionEOSGS */
 
-FSessionEOSGS::FSessionEOSGS(const TSharedPtr<FSessionDetailsHandleEOSGS>& InSessionDetailsHandle)
+FSessionEOSGS::FSessionEOSGS(const FSessionsEOSGS& SessionsInstance, const TSharedPtr<FSessionDetailsHandleEOSGS>& InSessionDetailsHandle)
 	: SessionDetailsHandle(InSessionDetailsHandle)
 {
 	EOS_SessionDetails_CopyInfoOptions CopyInfoOptions = { };
@@ -54,7 +56,7 @@ FSessionEOSGS::FSessionEOSGS(const TSharedPtr<FSessionDetailsHandleEOSGS>& InSes
 	EOS_EResult CopyInfoResult = EOS_SessionDetails_CopyInfo(SessionDetailsHandle->SessionDetailsHandle, &CopyInfoOptions, &SessionDetailsInfo);
 	if (CopyInfoResult == EOS_EResult::EOS_Success)
 	{
-		SessionInfo.SessionId = FSessionsEOSGS::CreateSessionId(FString(SessionDetailsInfo->SessionId));
+		SessionInfo.SessionId = SessionsInstance.CreateSessionId(UTF8_TO_TCHAR(SessionDetailsInfo->SessionId));
 		SessionSettings.CustomSettings.Emplace(EOSGS_HOST_ADDRESS_ATTRIBUTE_KEY, FCustomSessionSetting{ FSchemaVariant(FString(UTF8_TO_TCHAR(SessionDetailsInfo->HostAddress))), ESchemaAttributeVisibility::Public });
 
 		// bInvitesAllowed and bAllowJoinInProgress should both have the same value, and that value corresponds to bAllowNewMembers
@@ -135,7 +137,7 @@ const FSessionEOSGS& FSessionEOSGS::Cast(const ISession& InSession)
 
 /** FSessionsEOSGS */
 
-FSessionsEOSGS::FSessionsEOSGS(FOnlineServicesEOSGS& InServices)
+FSessionsEOSGS::FSessionsEOSGS(FOnlineServicesEpicCommon& InServices)
 	: Super(InServices)
 {
 }
@@ -144,7 +146,7 @@ void FSessionsEOSGS::Initialize()
 {
 	Super::Initialize();
 
-	SessionsHandle = EOS_Platform_GetSessionsInterface(*GetServices<FOnlineServicesEOSGS>().GetEOSPlatformHandle());
+	SessionsHandle = EOS_Platform_GetSessionsInterface(*GetServices<FOnlineServicesEpicCommon>().GetEOSPlatformHandle());
 	check(SessionsHandle);
 
 	RegisterEventHandlers();
@@ -1005,8 +1007,8 @@ void FSessionsEOSGS::SetSessionSearchSessionId(FSessionSearchHandleEOSGS& Sessio
 	Options.ApiVersion = 1;
 	UE_EOS_CHECK_API_MISMATCH(EOS_SESSIONSEARCH_SETSESSIONID_API_LATEST, 1);
 
-	const FTCHARToUTF8 SessionIdUtf8(*FOnlineSessionIdRegistryEOSGS::Get().ToLogString(SessionId));
-	Options.SessionId = SessionIdUtf8.Get();
+	const auto SessionIdUtf8 = StringCast<UTF8CHAR>(*ToString(SessionId));
+	Options.SessionId = (const char*)SessionIdUtf8.Get();
 
 	EOS_EResult ResultCode = EOS_SessionSearch_SetSessionId(SessionSearchHandle.SearchHandle, &Options);
 	if (ResultCode != EOS_EResult::EOS_Success)
@@ -1049,14 +1051,14 @@ void FSessionsEOSGS::WriteSessionSearchHandle(FSessionSearchHandleEOSGS& Session
 	}
 }
 
-FOnlineSessionId FSessionsEOSGS::CreateSessionId(const FString& SessionId)
+FOnlineSessionId FSessionsEOSGS::CreateSessionId(const FString& SessionId) const
 {
-	return FOnlineSessionIdRegistryEOSGS::Get().BasicRegistry.FindOrAddHandle(SessionId);
+	return FOnlineSessionIdRegistryEOSGS::GetRegistered(Services.GetServicesProvider()).BasicRegistry.FindOrAddHandle(SessionId);
 }
 
 FSessionInviteId FSessionsEOSGS::CreateSessionInviteId(const FString& SessionInviteId) const
 {
-	return FOnlineSessionInviteIdRegistryEOSGS::Get().BasicRegistry.FindOrAddHandle(SessionInviteId);
+	return FOnlineSessionInviteIdRegistryEOSGS::GetRegistered(Services.GetServicesProvider()).BasicRegistry.FindOrAddHandle(SessionInviteId);
 }
 
 TFuture<TOnlineResult<FFindSessions>> FSessionsEOSGS::FindSessionsImpl(const FFindSessions::Params& Params)
@@ -1201,7 +1203,7 @@ TFuture<TOnlineResult<FJoinSession>> FSessionsEOSGS::JoinSessionImpl(const FJoin
 	if (GetSessionByIdResult.IsError())
 	{
 		// If no result is found, the id might be expired, which we should notify
-		if (FOnlineSessionIdRegistryEOSGS::Get().IsSessionIdExpired(Params.SessionId))
+		if (FOnlineSessionIdRegistryEOSGS::GetRegistered(Services.GetServicesProvider()).IsSessionIdExpired(Params.SessionId))
 		{
 			UE_LOG(LogOnlineServices, Warning, TEXT("[%s] SessionId parameter [%s] is expired. Please call FindSessions to get an updated list of available sessions "), UTF8_TO_TCHAR(__FUNCTION__), *ToLogString(Params.SessionId));
 		}
@@ -1249,7 +1251,7 @@ TFuture<TOnlineResult<FJoinSession>> FSessionsEOSGS::JoinSessionImpl(const FJoin
 			if (GetSessionByIdResult.IsError())
 			{
 				// If no result is found, the id might be expired, which we should notify
-				if (FOnlineSessionIdRegistryEOSGS::Get().IsSessionIdExpired(Params.SessionId))
+				if (FOnlineSessionIdRegistryEOSGS::GetRegistered(Services.GetServicesProvider()).IsSessionIdExpired(Params.SessionId))
 				{
 					UE_LOG(LogOnlineServices, Warning, TEXT("[%s] SessionId parameter [%s] is expired. Please call FindSessions to get an updated list of available sessions "), UTF8_TO_TCHAR(__FUNCTION__), *ToLogString(Params.SessionId));
 				}
@@ -1397,7 +1399,7 @@ TOnlineAsyncOpHandle<FBuildSessionFromDetailsHandle> FSessionsEOSGS::BuildSessio
 			if (TOnlineAsyncOpPtr<FBuildSessionFromDetailsHandle> StrongOp = WeakOp.Pin())
 			{
 				// After all the ids are resolved, we can build the session safely
-				StrongOp->SetResult(FBuildSessionFromDetailsHandle::Result{ StrongOp->GetParams().LocalAccountId, MakeShared<FSessionEOSGS>(Params.SessionDetailsHandleEOSGS) });
+				StrongOp->SetResult(FBuildSessionFromDetailsHandle::Result{ StrongOp->GetParams().LocalAccountId, MakeShared<FSessionEOSGS>(*this, Params.SessionDetailsHandleEOSGS) });
 			}
 		});
 	}
@@ -1468,7 +1470,7 @@ TFuture<TOnlineResult<FRejectSessionInvite>> FSessionsEOSGS::RejectSessionInvite
 
 	RejectInviteOptions.LocalUserId = GetProductUserIdChecked(Params.LocalAccountId);
 
-	const FString& InviteIdStr = FOnlineSessionInviteIdRegistryEOSGS::Get().BasicRegistry.FindIdValue(Params.SessionInviteId);
+	const FString& InviteIdStr = FOnlineSessionInviteIdRegistryEOSGS::GetRegistered(Services.GetServicesProvider()).BasicRegistry.FindIdValue(Params.SessionInviteId);
 	const FTCHARToUTF8 InviteIdUtf8(*InviteIdStr);
 	RejectInviteOptions.InviteId = InviteIdUtf8.Get();
 
diff --git a/Engine/Plugins/Online/OnlineServicesEOSGS/Source/Private/Online/StatsEOSGS.cpp b/Engine/Plugins/Online/OnlineServicesEOSGS/Source/Private/Online/StatsEOSGS.cpp
--- a/Engine/Plugins/Online/OnlineServicesEOSGS/Source/Private/Online/StatsEOSGS.cpp
+++ b/Engine/Plugins/Online/OnlineServicesEOSGS/Source/Private/Online/StatsEOSGS.cpp
@@ -21,7 +21,7 @@ struct FStatNameRawBufferEOS
 	char StatName[UE_ONLINE_STAT_EOS_NAME_MAX_LENGTH];
 };
 
-FStatsEOSGS::FStatsEOSGS(FOnlineServicesEOSGS& InServices)
+FStatsEOSGS::FStatsEOSGS(FOnlineServicesEpicCommon& InServices)
 	: Super(InServices)
 {
 }
@@ -30,7 +30,7 @@ void FStatsEOSGS::Initialize()
 {
 	Super::Initialize();
 
-	StatsHandle = EOS_Platform_GetStatsInterface(*GetServices<FOnlineServicesEOSGS>().GetEOSPlatformHandle());
+	StatsHandle = EOS_Platform_GetStatsInterface(*GetServices<FOnlineServicesEpicCommon>().GetEOSPlatformHandle());
 	check(StatsHandle);
 }
 
diff --git a/Engine/Plugins/Online/OnlineServicesEOSGS/Source/Private/Online/TitleFileEOSGS.cpp b/Engine/Plugins/Online/OnlineServicesEOSGS/Source/Private/Online/TitleFileEOSGS.cpp
--- a/Engine/Plugins/Online/OnlineServicesEOSGS/Source/Private/Online/TitleFileEOSGS.cpp
+++ b/Engine/Plugins/Online/OnlineServicesEOSGS/Source/Private/Online/TitleFileEOSGS.cpp
@@ -14,7 +14,7 @@
 
 namespace UE::Online {
 
-FTitleFileEOSGS::FTitleFileEOSGS(FOnlineServicesEOSGS& InServices)
+FTitleFileEOSGS::FTitleFileEOSGS(FOnlineServicesEpicCommon& InServices)
 	: Super(InServices)
 {
 }
@@ -23,7 +23,7 @@ void FTitleFileEOSGS::Initialize()
 {
 	Super::Initialize();
 
-	TitleStorageHandle = EOS_Platform_GetTitleStorageInterface(*GetServices<FOnlineServicesEOSGS>().GetEOSPlatformHandle());
+	TitleStorageHandle = EOS_Platform_GetTitleStorageInterface(*GetServices<FOnlineServicesEpicCommon>().GetEOSPlatformHandle());
 	check(TitleStorageHandle);
 }
 
diff --git a/Engine/Plugins/Online/OnlineServicesEOSGS/Source/Private/Online/UserFileEOSGS.cpp b/Engine/Plugins/Online/OnlineServicesEOSGS/Source/Private/Online/UserFileEOSGS.cpp
--- a/Engine/Plugins/Online/OnlineServicesEOSGS/Source/Private/Online/UserFileEOSGS.cpp
+++ b/Engine/Plugins/Online/OnlineServicesEOSGS/Source/Private/Online/UserFileEOSGS.cpp
@@ -16,7 +16,7 @@ namespace UE::Online {
 
 static const TCHAR* RequestHandleKey = TEXT("RequestHandle");
 
-FUserFileEOSGS::FUserFileEOSGS(FOnlineServicesEOSGS& InServices)
+FUserFileEOSGS::FUserFileEOSGS(FOnlineServicesEpicCommon& InServices)
 	: Super(InServices)
 {
 }
@@ -25,7 +25,7 @@ void FUserFileEOSGS::Initialize()
 {
 	Super::Initialize();
 
-	PlayerDataStorageHandle = EOS_Platform_GetPlayerDataStorageInterface(*GetServices<FOnlineServicesEOSGS>().GetEOSPlatformHandle());
+	PlayerDataStorageHandle = EOS_Platform_GetPlayerDataStorageInterface(*GetServices<FOnlineServicesEpicCommon>().GetEOSPlatformHandle());
 	check(PlayerDataStorageHandle);
 }
 
diff --git a/Engine/Plugins/Online/OnlineServicesEOSGS/Source/Public/Online/AchievementsEOSGS.h b/Engine/Plugins/Online/OnlineServicesEOSGS/Source/Public/Online/AchievementsEOSGS.h
--- a/Engine/Plugins/Online/OnlineServicesEOSGS/Source/Public/Online/AchievementsEOSGS.h
+++ b/Engine/Plugins/Online/OnlineServicesEOSGS/Source/Public/Online/AchievementsEOSGS.h
@@ -2,7 +2,6 @@
 
 #pragma once
 
-#include "CoreMinimal.h"
 #include "Online/AchievementsCommon.h"
 #include "Online/OnlineServicesEOSGSTypes.h"
 
@@ -13,14 +12,14 @@
 
 namespace UE::Online {
 
-class FOnlineServicesEOSGS;
+class FOnlineServicesEpicCommon;
 
 class ONLINESERVICESEOSGS_API FAchievementsEOSGS : public FAchievementsCommon
 {
 public:
 	using Super = FAchievementsCommon;
 
-	FAchievementsEOSGS(FOnlineServicesEOSGS& InOwningSubsystem);
+	FAchievementsEOSGS(FOnlineServicesEpicCommon& InOwningSubsystem);
 	virtual ~FAchievementsEOSGS() = default;
 
 	// IOnlineComponent
diff --git a/Engine/Plugins/Online/OnlineServicesEOSGS/Source/Public/Online/AuthEOSGS.h b/Engine/Plugins/Online/OnlineServicesEOSGS/Source/Public/Online/AuthEOSGS.h
--- a/Engine/Plugins/Online/OnlineServicesEOSGS/Source/Public/Online/AuthEOSGS.h
+++ b/Engine/Plugins/Online/OnlineServicesEOSGS/Source/Public/Online/AuthEOSGS.h
@@ -2,7 +2,6 @@
 
 #pragma once
 
-#include "CoreMinimal.h"
 #include "Online/AuthCommon.h"
 #include "Online/OnlineServicesEOSGSTypes.h"
 
diff --git a/Engine/Plugins/Online/OnlineServicesEOSGS/Source/Public/Online/ExternalUIEOSGS.h b/Engine/Plugins/Online/OnlineServicesEOSGS/Source/Public/Online/ExternalUIEOSGS.h
--- a/Engine/Plugins/Online/OnlineServicesEOSGS/Source/Public/Online/ExternalUIEOSGS.h
+++ b/Engine/Plugins/Online/OnlineServicesEOSGS/Source/Public/Online/ExternalUIEOSGS.h
@@ -2,15 +2,15 @@
 
 #pragma once
 
-#include "CoreMinimal.h"
 #include "Online/ExternalUICommon.h"
-#include "OnlineServicesEOSGSTypes.h"
+#include "Online/OnlineServicesEOSGSTypes.h"
+#include "Online/OnlineServicesEpicCommon.h"
 
 #include "eos_ui_types.h"
 
 namespace UE::Online {
 
-class FOnlineServicesEOSGS;
+class FOnlineServicesEpicCommon;
 
 struct FExternalUIProcessDisplaySettingsUpdatedImp
 {
@@ -37,7 +37,7 @@ class ONLINESERVICESEOSGS_API FExternalUIEOSGS : public FExternalUICommon
 public:
 	using Super = FExternalUICommon;
 
-	FExternalUIEOSGS(FOnlineServicesEOSGS& InOwningSubsystem);
+	FExternalUIEOSGS(FOnlineServicesEpicCommon& InOwningSubsystem);
 
 	virtual void Initialize() override;
 	virtual void PreShutdown() override;
diff --git a/Engine/Plugins/Online/OnlineServicesEOSGS/Source/Public/Online/LeaderboardsEOSGS.h b/Engine/Plugins/Online/OnlineServicesEOSGS/Source/Public/Online/LeaderboardsEOSGS.h
--- a/Engine/Plugins/Online/OnlineServicesEOSGS/Source/Public/Online/LeaderboardsEOSGS.h
+++ b/Engine/Plugins/Online/OnlineServicesEOSGS/Source/Public/Online/LeaderboardsEOSGS.h
@@ -2,7 +2,6 @@
 
 #pragma once
 
-#include "CoreMinimal.h"
 #include "Online/LeaderboardsCommon.h"
 
 #if defined(EOS_PLATFORM_BASE_FILE_NAME)
@@ -12,14 +11,14 @@
 
 namespace UE::Online {
 
-class FOnlineServicesEOSGS;
+class FOnlineServicesEpicCommon;
 
 class ONLINESERVICESEOSGS_API FLeaderboardsEOSGS : public FLeaderboardsCommon
 {
 public:
 	using Super = FLeaderboardsCommon;
 
-	FLeaderboardsEOSGS(FOnlineServicesEOSGS& InOwningSubsystem);
+	FLeaderboardsEOSGS(FOnlineServicesEpicCommon& InOwningSubsystem);
 
 	// TOnlineComponent
 	virtual void Initialize() override;
diff --git a/Engine/Plugins/Online/OnlineServicesEOSGS/Source/Public/Online/LobbiesEOSGS.h b/Engine/Plugins/Online/OnlineServicesEOSGS/Source/Public/Online/LobbiesEOSGS.h
--- a/Engine/Plugins/Online/OnlineServicesEOSGS/Source/Public/Online/LobbiesEOSGS.h
+++ b/Engine/Plugins/Online/OnlineServicesEOSGS/Source/Public/Online/LobbiesEOSGS.h
@@ -2,7 +2,6 @@
 
 #pragma once
 
-#include "CoreMinimal.h"
 #include "Online/LobbiesCommon.h"
 #include "OnlineServicesEOSGSTypes.h"
 
@@ -17,7 +16,7 @@ struct FLobbyPrerequisitesEOS;
 class FLobbyDataEOS;
 class FLobbyDataRegistryEOS;
 class FLobbyInviteDataEOS;
-class FOnlineServicesEOSGS;
+class FOnlineServicesEpicCommon;
 class FLobbySearchEOS;
 struct FClientLobbyDataChanges;
 struct FClientLobbyMemberDataChanges;
@@ -189,7 +188,7 @@ class FLobbiesEOSGS : public FLobbiesCommon
 public:
 	using Super = FLobbiesCommon;
 
-	FLobbiesEOSGS(FOnlineServicesEOSGS& InServices);
+	FLobbiesEOSGS(FOnlineServicesEpicCommon& InServices);
 
 	virtual void Initialize() override;
 	virtual void PreShutdown() override;
diff --git a/Engine/Plugins/Online/OnlineServicesEOSGS/Source/Public/Online/NboSerializerEOSGSSvc.h b/Engine/Plugins/Online/OnlineServicesEOSGS/Source/Public/Online/NboSerializerEOSGSSvc.h
--- a/Engine/Plugins/Online/OnlineServicesEOSGS/Source/Public/Online/NboSerializerEOSGSSvc.h
+++ b/Engine/Plugins/Online/OnlineServicesEOSGS/Source/Public/Online/NboSerializerEOSGSSvc.h
@@ -20,16 +20,27 @@ namespace NboSerializerEOSGSSvc {
 
 /** NboSerializeToBuffer methods */
 
+inline void SerializeToBuffer(FNboSerializeToBuffer& Ar, const EOnlineServices ServicesType)
+{
+	Ar << (__underlying_type(EOnlineServices))ServicesType;
+}
+
 inline void SerializeToBuffer(FNboSerializeToBuffer& Ar, const FAccountId& UniqueId)
 {
-	TArray<uint8> Data = FOnlineAccountIdRegistryEOSGS::GetRegistered(UniqueId.GetOnlineServicesType()).ToReplicationData(UniqueId);
+	const TArray<uint8> Data = FOnlineAccountIdRegistryEOSGS::GetRegistered(UniqueId.GetOnlineServicesType()).ToReplicationData(UniqueId);
+
+	SerializeToBuffer(Ar, UniqueId.GetOnlineServicesType());
+
 	Ar << Data.Num();
 	Ar.WriteBinary(Data.GetData(), Data.Num());
 }
 
 inline void SerializeToBuffer(FNboSerializeToBuffer& Ar, const FOnlineSessionId& SessionId)
 {
-	TArray<uint8> Data = FOnlineSessionIdRegistryEOSGS::Get().ToReplicationData(SessionId);
+	const TArray<uint8> Data = FOnlineSessionIdRegistryEOSGS::GetRegistered(SessionId.GetOnlineServicesType()).ToReplicationData(SessionId);
+
+	SerializeToBuffer(Ar, SessionId.GetOnlineServicesType());
+
 	Ar << Data.Num();
 	Ar.WriteBinary(Data.GetData(), Data.Num());
 }
@@ -37,7 +48,7 @@ inline void SerializeToBuffer(FNboSerializeToBuffer& Ar, const FOnlineSessionId&
 inline void SerializeToBuffer(FNboSerializeToBuffer& Packet, const FSessionMemberIdsSet& SessionMembersSet)
 {
 	Packet << SessionMembersSet.Num();
-
+	
 	for (const FAccountId& SessionMember : SessionMembersSet)
 	{
 		SerializeToBuffer(Packet, SessionMember);
@@ -46,22 +57,41 @@ inline void SerializeToBuffer(FNboSerializeToBuffer& Packet, const FSessionMembe
 
 /** NboSerializeFromBuffer methods */
 
+inline void SerializeFromBuffer(FNboSerializeFromBuffer& Ar, EOnlineServices& ServicesType)
+{
+	using UnderlyingType = __underlying_type(EOnlineServices);
+	UnderlyingType Value;
+	Ar >> Value;
+
+	ServicesType = (EOnlineServices)Value;
+}
+
 inline void SerializeFromBuffer(FNboSerializeFromBuffer& Ar, FAccountId& UniqueId)
 {
-	TArray<uint8> Data;
+	EOnlineServices ServicesType;
+	SerializeFromBuffer(Ar, ServicesType);
+
 	int32 Size;
 	Ar >> Size;
+
+	TArray<uint8> Data;
 	Ar.ReadBinaryArray(Data, Size);
-	UniqueId = FOnlineAccountIdRegistryEOSGS::GetRegistered(UniqueId.GetOnlineServicesType()).FromReplicationData(Data);
+
+	UniqueId = FOnlineAccountIdRegistryEOSGS::GetRegistered(ServicesType).FromReplicationData(Data);
 }
 
 inline void SerializeFromBuffer(FNboSerializeFromBuffer& Ar, FOnlineSessionId& SessionId)
 {
-	TArray<uint8> Data;
+	EOnlineServices ServicesType;
+	SerializeFromBuffer(Ar, ServicesType);
+
 	int32 Size;
 	Ar >> Size;
+
+	TArray<uint8> Data;
 	Ar.ReadBinaryArray(Data, Size);
-	SessionId = FOnlineSessionIdRegistryEOSGS::Get().FromReplicationData(Data);
+
+	SessionId = FOnlineSessionIdRegistryEOSGS::GetRegistered(ServicesType).FromReplicationData(Data);
 }
 
 inline void SerializeFromBuffer(FNboSerializeFromBuffer& Packet, FSessionMemberIdsSet& SessionMembersSet)
diff --git a/Engine/Plugins/Online/OnlineServicesEOSGS/Source/Public/Online/OnlineServicesEOSGS.h b/Engine/Plugins/Online/OnlineServicesEOSGS/Source/Public/Online/OnlineServicesEOSGS.h
--- a/Engine/Plugins/Online/OnlineServicesEOSGS/Source/Public/Online/OnlineServicesEOSGS.h
+++ b/Engine/Plugins/Online/OnlineServicesEOSGS/Source/Public/Online/OnlineServicesEOSGS.h
@@ -42,7 +42,6 @@ public:
 	static const TCHAR* GetServiceConfigNameStatic() { return TEXT("EOS"); }
 
 protected:
-	void WarnIfEncryptionKeyMissing(const FString& InterfaceName) const;
 
 #if WITH_ENGINE
 	TSharedPtr<FSocketSubsystemEOS, ESPMode::ThreadSafe> SocketSubsystem;
diff --git a/Engine/Plugins/Online/OnlineServicesEOSGS/Source/Public/Online/SessionsEOSGS.h b/Engine/Plugins/Online/OnlineServicesEOSGS/Source/Public/Online/SessionsEOSGS.h
--- a/Engine/Plugins/Online/OnlineServicesEOSGS/Source/Public/Online/SessionsEOSGS.h
+++ b/Engine/Plugins/Online/OnlineServicesEOSGS/Source/Public/Online/SessionsEOSGS.h
@@ -14,23 +14,35 @@
 
 namespace UE::Online {
 
-class FOnlineServicesEOSGS;
+class FOnlineServicesEpicCommon;
+class FSessionsEOSGS;
 
 class FOnlineSessionIdRegistryEOSGS : public FOnlineSessionIdRegistryLAN
 {
 public:
-	static FOnlineSessionIdRegistryEOSGS& Get();
+	UE_DEPRECATED(5.6, "This method is deprecated, please use GetRegistered(EOnlineServices)")
+	static FOnlineSessionIdRegistryEOSGS& Get() { return GetRegistered(EOnlineServices::Epic); }
+	
+	static FOnlineSessionIdRegistryEOSGS& GetRegistered(EOnlineServices Services);
+	
 private:
-	FOnlineSessionIdRegistryEOSGS();
+	friend class FOnlineServicesEOSGSModule;
+	friend class FOnlineServicesEpicGameModule;
+	FOnlineSessionIdRegistryEOSGS(EOnlineServices ServicesType);
 };
 
 class FOnlineSessionInviteIdRegistryEOSGS : public FOnlineSessionInviteIdStringRegistry
 {
 public:
-	static FOnlineSessionInviteIdRegistryEOSGS& Get();
+	UE_DEPRECATED(5.6, "This method is deprecated, please use GetRegistered(EOnlineServices)")
+	static FOnlineSessionInviteIdRegistryEOSGS& Get() { return GetRegistered(EOnlineServices::Epic); }
+
+	static FOnlineSessionInviteIdRegistryEOSGS& GetRegistered(EOnlineServices Services);
 
 private:
-	FOnlineSessionInviteIdRegistryEOSGS();
+	friend class FOnlineServicesEOSGSModule;
+	friend class FOnlineServicesEpicGameModule;
+	FOnlineSessionInviteIdRegistryEOSGS(EOnlineServices ServicesType);
 };
 
 struct FSessionModificationHandleEOSGS : FNoncopyable
@@ -85,7 +97,7 @@ public:
 	FSessionEOSGS(const FSessionEOSGS& InSession) = default;
 
 	/** This constructor should only be used by BuildSessionFromDetailsHandle, after all user ids in the session have been resolved. */
-	FSessionEOSGS(const TSharedPtr<FSessionDetailsHandleEOSGS>& InSessionDetailsHandle);
+	FSessionEOSGS(const FSessionsEOSGS& SessionsInstance, const TSharedPtr<FSessionDetailsHandleEOSGS>& InSessionDetailsHandle);
 
 	static const FSessionEOSGS& Cast(const ISession& InSession);
 
@@ -182,7 +194,7 @@ public:
 
 	using Super = FSessionsLAN;
 
-	FSessionsEOSGS(FOnlineServicesEOSGS& InOwningSubsystem);
+	FSessionsEOSGS(FOnlineServicesEpicCommon& InOwningSubsystem);
 	virtual ~FSessionsEOSGS() = default;
 
 	// TOnlineComponent
@@ -256,7 +268,7 @@ protected:
 	 */
 	TFuture<TDefaultErrorResult<FSendSingleSessionInviteImpl>> SendSingleSessionInviteImpl(FSendSingleSessionInviteImpl::Params&& Params);
 
-	static FOnlineSessionId CreateSessionId(const FString& SessionId);
+	FOnlineSessionId CreateSessionId(const FString& SessionId) const;
 	FSessionInviteId CreateSessionInviteId(const FString& SessionInviteId) const;
 
 	/**
diff --git a/Engine/Plugins/Online/OnlineServicesEOSGS/Source/Public/Online/StatsEOSGS.h b/Engine/Plugins/Online/OnlineServicesEOSGS/Source/Public/Online/StatsEOSGS.h
--- a/Engine/Plugins/Online/OnlineServicesEOSGS/Source/Public/Online/StatsEOSGS.h
+++ b/Engine/Plugins/Online/OnlineServicesEOSGS/Source/Public/Online/StatsEOSGS.h
@@ -2,7 +2,6 @@
 
 #pragma once
 
-#include "CoreMinimal.h"
 #include "Online/StatsCommon.h"
 
 #if defined(EOS_PLATFORM_BASE_FILE_NAME)
@@ -12,7 +11,7 @@
 
 namespace UE::Online {
 
-class FOnlineServicesEOSGS;
+class FOnlineServicesEpicCommon;
 
 /**
  * Because the value type is int32 in EOS, stats types are converted to int32. So these are limitations: 
@@ -25,7 +24,7 @@ class ONLINESERVICESEOSGS_API FStatsEOSGS : public FStatsCommon
 public:
 	using Super = FStatsCommon;
 
-	FStatsEOSGS(FOnlineServicesEOSGS& InOwningSubsystem);
+	FStatsEOSGS(FOnlineServicesEpicCommon& InOwningSubsystem);
 
 	// TOnlineComponent
 	virtual void Initialize() override;
diff --git a/Engine/Plugins/Online/OnlineServicesEOSGS/Source/Public/Online/TitleFileEOSGS.h b/Engine/Plugins/Online/OnlineServicesEOSGS/Source/Public/Online/TitleFileEOSGS.h
--- a/Engine/Plugins/Online/OnlineServicesEOSGS/Source/Public/Online/TitleFileEOSGS.h
+++ b/Engine/Plugins/Online/OnlineServicesEOSGS/Source/Public/Online/TitleFileEOSGS.h
@@ -2,7 +2,6 @@
 
 #pragma once
 
-#include "CoreMinimal.h"
 #include "Online/TitleFileCommon.h"
 
 #if defined(EOS_PLATFORM_BASE_FILE_NAME)
@@ -12,7 +11,7 @@
 
 namespace UE::Online {
 
-class FOnlineServicesEOSGS;
+class FOnlineServicesEpicCommon;
 
 struct FTitleFileEOSGSConfig
 {
@@ -34,7 +33,7 @@ class ONLINESERVICESEOSGS_API FTitleFileEOSGS : public FTitleFileCommon
 public:
 	using Super = FTitleFileCommon;
 
-	FTitleFileEOSGS(FOnlineServicesEOSGS& InOwningSubsystem);
+	FTitleFileEOSGS(FOnlineServicesEpicCommon& InOwningSubsystem);
 	virtual ~FTitleFileEOSGS() = default;
 
 	// IOnlineComponent
diff --git a/Engine/Plugins/Online/OnlineServicesEOSGS/Source/Public/Online/UserFileEOSGS.h b/Engine/Plugins/Online/OnlineServicesEOSGS/Source/Public/Online/UserFileEOSGS.h
--- a/Engine/Plugins/Online/OnlineServicesEOSGS/Source/Public/Online/UserFileEOSGS.h
+++ b/Engine/Plugins/Online/OnlineServicesEOSGS/Source/Public/Online/UserFileEOSGS.h
@@ -2,7 +2,6 @@
 
 #pragma once
 
-#include "CoreMinimal.h"
 #include "Online/UserFileCommon.h"
 
 #if defined(EOS_PLATFORM_BASE_FILE_NAME)
@@ -12,7 +11,7 @@
 
 namespace UE::Online {
 
-class FOnlineServicesEOSGS;
+class FOnlineServicesEpicCommon;
 
 struct FUserFileEOSGSConfig
 {
@@ -32,7 +31,7 @@ class ONLINESERVICESEOSGS_API FUserFileEOSGS : public FUserFileCommon
 public:
 	using Super = FUserFileCommon;
 
-	FUserFileEOSGS(FOnlineServicesEOSGS& InOwningSubsystem);
+	FUserFileEOSGS(FOnlineServicesEpicCommon& InOwningSubsystem);
 	virtual ~FUserFileEOSGS() = default;
 
 	// IOnlineComponent
diff --git a/Engine/Plugins/Online/OnlineServicesEpicCommon/Source/OnlineServicesEpicCommon/Private/Online/OnlineServicesEpicCommon.cpp b/Engine/Plugins/Online/OnlineServicesEpicCommon/Source/OnlineServicesEpicCommon/Private/Online/OnlineServicesEpicCommon.cpp
--- a/Engine/Plugins/Online/OnlineServicesEpicCommon/Source/OnlineServicesEpicCommon/Private/Online/OnlineServicesEpicCommon.cpp
+++ b/Engine/Plugins/Online/OnlineServicesEpicCommon/Source/OnlineServicesEpicCommon/Private/Online/OnlineServicesEpicCommon.cpp
@@ -69,6 +69,27 @@ void FOnlineServicesEpicCommon::RemoveEOSSDKFastTick(FOnlineAsyncOp& InAsyncOp)
 	}
 }
 
+void FOnlineServicesEpicCommon::WarnIfEncryptionKeyMissing(const FString& InterfaceName) const
+{
+	if (IEOSSDKManager* Manager = IEOSSDKManager::Get())
+	{
+		const FString& PlatformConfigName = GetEOSPlatformHandle()->GetConfigName();
+		if (const FEOSSDKPlatformConfig* Config = Manager->GetPlatformConfig(PlatformConfigName))
+		{
+			const FString& EncryptionKey = Config->EncryptionKey;
+			if (EncryptionKey.IsEmpty())
+			{
+				UE_LOG(LogOnlineServices, Verbose, TEXT("%s interface not available due to missing ClientEncryptionKey in config."), *InterfaceName);
+			}
+			else
+			{
+				// If we have an encryption key and still can't get the interface, something weird is going on.
+				UE_LOG(LogOnlineServices, Warning, TEXT("%s interface not available despite encryption key being present."), *InterfaceName);
+			}
+		}
+	}
+}
+
 void FOnlineServicesEpicCommon::FlushTick(float DeltaSeconds)
 {
 	Super::FlushTick(DeltaSeconds);
diff --git a/Engine/Plugins/Online/OnlineServicesEpicCommon/Source/OnlineServicesEpicCommon/Public/Online/OnlineServicesEpicCommon.h b/Engine/Plugins/Online/OnlineServicesEpicCommon/Source/OnlineServicesEpicCommon/Public/Online/OnlineServicesEpicCommon.h
--- a/Engine/Plugins/Online/OnlineServicesEpicCommon/Source/OnlineServicesEpicCommon/Public/Online/OnlineServicesEpicCommon.h
+++ b/Engine/Plugins/Online/OnlineServicesEpicCommon/Source/OnlineServicesEpicCommon/Public/Online/OnlineServicesEpicCommon.h
@@ -27,6 +27,8 @@ public:
 	void RemoveEOSSDKFastTick(FOnlineAsyncOp& InAsyncOp);
 
 protected:
+	void WarnIfEncryptionKeyMissing(const FString& InterfaceName) const;
+	
 	virtual void FlushTick(float DeltaSeconds) override;
 
 	IEOSPlatformHandlePtr EOSPlatformHandle;
