[Backout] - CL41316842
#rnx
Original CL Desc
-----------------------------------------------------------------
- Register more EOS[GS] components in the split plugins

#rnx

diff --git a/Engine/Plugins/Online/OnlineServicesEOS/Source/Private/Online/ExternalUIEOS.cpp b/Engine/Plugins/Online/OnlineServicesEOS/Source/Private/Online/ExternalUIEOS.cpp
--- a/Engine/Plugins/Online/OnlineServicesEOS/Source/Private/Online/ExternalUIEOS.cpp
+++ b/Engine/Plugins/Online/OnlineServicesEOS/Source/Private/Online/ExternalUIEOS.cpp
@@ -12,7 +12,7 @@
 
 namespace UE::Online {
 
-FExternalUIEOS::FExternalUIEOS(FOnlineServicesEpicCommon& InServices)
+FExternalUIEOS::FExternalUIEOS(FOnlineServicesEOS& InServices)
 	: Super(InServices)
 {
 }
@@ -21,7 +21,7 @@ void FExternalUIEOS::Initialize()
 {
 	Super::Initialize();
 
-	UIHandle = EOS_Platform_GetUIInterface(*GetServices<FOnlineServicesEpicCommon>().GetEOSPlatformHandle());
+	UIHandle = EOS_Platform_GetUIInterface(*GetServices<FOnlineServicesEOS>().GetEOSPlatformHandle());
 	check(UIHandle != nullptr);
 }
 
diff --git a/Engine/Plugins/Online/OnlineServicesEOS/Source/Private/Online/PresenceEOS.cpp b/Engine/Plugins/Online/OnlineServicesEOS/Source/Private/Online/PresenceEOS.cpp
--- a/Engine/Plugins/Online/OnlineServicesEOS/Source/Private/Online/PresenceEOS.cpp
+++ b/Engine/Plugins/Online/OnlineServicesEOS/Source/Private/Online/PresenceEOS.cpp
@@ -61,7 +61,7 @@ static inline EOS_Presence_EStatus ToEOS_Presence_EStatus(EUserPresenceStatus In
 	return EOS_Presence_EStatus::EOS_PS_Offline;
 }
 
-FPresenceEOS::FPresenceEOS(FOnlineServicesEpicCommon& InServices)
+FPresenceEOS::FPresenceEOS(FOnlineServicesEOS& InServices)
 	: Super(InServices)
 {
 }
@@ -70,7 +70,7 @@ void FPresenceEOS::Initialize()
 {
 	Super::Initialize();
 
-	PresenceHandle = EOS_Platform_GetPresenceInterface(*GetServices<FOnlineServicesEpicCommon>().GetEOSPlatformHandle());
+	PresenceHandle = EOS_Platform_GetPresenceInterface(*GetServices<FOnlineServicesEOS>().GetEOSPlatformHandle());
 	check(PresenceHandle != nullptr);
 
 	LoginStatusChangedHandle = Services.Get<IAuth>()->OnLoginStatusChanged().Add(this, &FPresenceEOS::HandleAuthLoginStatusChanged);
diff --git a/Engine/Plugins/Online/OnlineServicesEOS/Source/Private/Online/SocialEOS.cpp b/Engine/Plugins/Online/OnlineServicesEOS/Source/Private/Online/SocialEOS.cpp
--- a/Engine/Plugins/Online/OnlineServicesEOS/Source/Private/Online/SocialEOS.cpp
+++ b/Engine/Plugins/Online/OnlineServicesEOS/Source/Private/Online/SocialEOS.cpp
@@ -26,7 +26,7 @@ ERelationship EOSEFriendStatusToERelationship(EOS_EFriendsStatus Status)
 	return ERelationship::NotFriend;
 }
 
-FSocialEOS::FSocialEOS(FOnlineServicesEpicCommon& InServices)
+FSocialEOS::FSocialEOS(FOnlineServicesEOS& InServices)
 	: FSocialCommon(InServices)
 {
 }
@@ -35,10 +35,10 @@ void FSocialEOS::Initialize()
 {
 	FSocialCommon::Initialize();
 
-	FriendsHandle = EOS_Platform_GetFriendsInterface(*GetServices<FOnlineServicesEpicCommon>().GetEOSPlatformHandle());
+	FriendsHandle = EOS_Platform_GetFriendsInterface(*GetServices<FOnlineServicesEOS>().GetEOSPlatformHandle());
 	check(FriendsHandle != nullptr);
 
-	UIHandle = EOS_Platform_GetUIInterface(*GetServices<FOnlineServicesEpicCommon>().GetEOSPlatformHandle());
+	UIHandle = EOS_Platform_GetUIInterface(*GetServices<FOnlineServicesEOS>().GetEOSPlatformHandle());
 
 	OnFriendsUpdate = EOS_RegisterComponentEventHandler(
 		this,
diff --git a/Engine/Plugins/Online/OnlineServicesEOS/Source/Private/Online/UserInfoEOS.cpp b/Engine/Plugins/Online/OnlineServicesEOS/Source/Private/Online/UserInfoEOS.cpp
--- a/Engine/Plugins/Online/OnlineServicesEOS/Source/Private/Online/UserInfoEOS.cpp
+++ b/Engine/Plugins/Online/OnlineServicesEOS/Source/Private/Online/UserInfoEOS.cpp
@@ -13,7 +13,7 @@
 
 namespace UE::Online {
 
-FUserInfoEOS::FUserInfoEOS(FOnlineServicesEpicCommon& InServices)
+FUserInfoEOS::FUserInfoEOS(FOnlineServicesEOS& InServices)
 	: FUserInfoCommon(InServices)
 {
 }
@@ -22,7 +22,7 @@ void FUserInfoEOS::Initialize()
 {
 	Super::Initialize();
 
-	UserInfoHandle = EOS_Platform_GetUserInfoInterface(*GetServices<FOnlineServicesEpicCommon>().GetEOSPlatformHandle());
+	UserInfoHandle = EOS_Platform_GetUserInfoInterface(*GetServices<FOnlineServicesEOS>().GetEOSPlatformHandle());
 	check(UserInfoHandle != nullptr);
 }
 
diff --git a/Engine/Plugins/Online/OnlineServicesEOS/Source/Public/Online/ExternalUIEOS.h b/Engine/Plugins/Online/OnlineServicesEOS/Source/Public/Online/ExternalUIEOS.h
--- a/Engine/Plugins/Online/OnlineServicesEOS/Source/Public/Online/ExternalUIEOS.h
+++ b/Engine/Plugins/Online/OnlineServicesEOS/Source/Public/Online/ExternalUIEOS.h
@@ -2,6 +2,7 @@
 
 #pragma once
 
+#include "CoreMinimal.h"
 #include "Online/ExternalUIEOSGS.h"
 
 #if defined(EOS_PLATFORM_BASE_FILE_NAME)
@@ -17,14 +18,14 @@ struct FExternalUIEOSConfig
 	bool bShowLoginUIEnabled = true;
 };
 
-class FOnlineServicesEpicCommon;
+class FOnlineServicesEOS;
 
 class ONLINESERVICESEOS_API FExternalUIEOS : public FExternalUIEOSGS
 {
 public:
 	using Super = FExternalUIEOSGS;
 
-	FExternalUIEOS(FOnlineServicesEpicCommon& InOwningSubsystem);
+	FExternalUIEOS(FOnlineServicesEOS& InOwningSubsystem);
 	virtual void Initialize() override;
 	virtual void UpdateConfig() override;
 	virtual TOnlineAsyncOpHandle<FExternalUIShowLoginUI> ShowLoginUI(FExternalUIShowLoginUI::Params&& Params) override;
diff --git a/Engine/Plugins/Online/OnlineServicesEOS/Source/Public/Online/PresenceEOS.h b/Engine/Plugins/Online/OnlineServicesEOS/Source/Public/Online/PresenceEOS.h
--- a/Engine/Plugins/Online/OnlineServicesEOS/Source/Public/Online/PresenceEOS.h
+++ b/Engine/Plugins/Online/OnlineServicesEOS/Source/Public/Online/PresenceEOS.h
@@ -13,14 +13,14 @@
 
 namespace UE::Online {
 
-class FOnlineServicesEpicCommon;
+class FOnlineServicesEOS;
 	
 class FPresenceEOS : public FPresenceCommon
 {
 public:
 	using Super = FPresenceCommon;
 
-	FPresenceEOS(FOnlineServicesEpicCommon& InServices);
+	FPresenceEOS(FOnlineServicesEOS& InServices);
 
 	virtual void Initialize() override;
 	virtual void PreShutdown() override;
diff --git a/Engine/Plugins/Online/OnlineServicesEOS/Source/Public/Online/SocialEOS.h b/Engine/Plugins/Online/OnlineServicesEOS/Source/Public/Online/SocialEOS.h
--- a/Engine/Plugins/Online/OnlineServicesEOS/Source/Public/Online/SocialEOS.h
+++ b/Engine/Plugins/Online/OnlineServicesEOS/Source/Public/Online/SocialEOS.h
@@ -2,6 +2,7 @@
 
 #pragma once
 
+#include "CoreMinimal.h"
 #include "Online/OnlineServicesEOSGSTypes.h"
 #include "Online/SocialCommon.h"
 
@@ -13,14 +14,14 @@
 
 namespace UE::Online {
 
-class FOnlineServicesEpicCommon;
+class FOnlineServicesEOS;
 	
 class FSocialEOS : public FSocialCommon
 {
 public:
 	using Super = FSocialCommon;
 
-	FSocialEOS(FOnlineServicesEpicCommon& InServices);
+	FSocialEOS(FOnlineServicesEOS& InServices);
 
 	virtual void Initialize() override;
 	virtual void PreShutdown() override;
diff --git a/Engine/Plugins/Online/OnlineServicesEOS/Source/Public/Online/UserInfoEOS.h b/Engine/Plugins/Online/OnlineServicesEOS/Source/Public/Online/UserInfoEOS.h
--- a/Engine/Plugins/Online/OnlineServicesEOS/Source/Public/Online/UserInfoEOS.h
+++ b/Engine/Plugins/Online/OnlineServicesEOS/Source/Public/Online/UserInfoEOS.h
@@ -2,6 +2,7 @@
 
 #pragma once
 
+#include "CoreMinimal.h"
 #include "Online/UserInfoCommon.h"
 
 #if defined(EOS_PLATFORM_BASE_FILE_NAME)
@@ -11,14 +12,14 @@
 
 namespace UE::Online {
 
-class FOnlineServicesEpicCommon;
+class FOnlineServicesEOS;
 	
 class FUserInfoEOS : public FUserInfoCommon
 {
 public:
 	using Super = FUserInfoCommon;
 
-	FUserInfoEOS(FOnlineServicesEpicCommon& InServices);
+	FUserInfoEOS(FOnlineServicesEOS& InServices);
 	virtual ~FUserInfoEOS() = default;
 
 	// IOnlineComponent
diff --git a/Engine/Plugins/Online/OnlineServicesEOSGS/Source/Private/Online/AchievementsEOSGS.cpp b/Engine/Plugins/Online/OnlineServicesEOSGS/Source/Private/Online/AchievementsEOSGS.cpp
--- a/Engine/Plugins/Online/OnlineServicesEOSGS/Source/Private/Online/AchievementsEOSGS.cpp
+++ b/Engine/Plugins/Online/OnlineServicesEOSGS/Source/Private/Online/AchievementsEOSGS.cpp
@@ -14,7 +14,7 @@
 
 namespace UE::Online {
 
-	FAchievementsEOSGS::FAchievementsEOSGS(FOnlineServicesEpicCommon& InServices)
+FAchievementsEOSGS::FAchievementsEOSGS(FOnlineServicesEOSGS& InServices)
 	: Super(InServices)
 {
 }
@@ -23,7 +23,7 @@ void FAchievementsEOSGS::Initialize()
 {
 	Super::Initialize();
 
-	AchievementsHandle = EOS_Platform_GetAchievementsInterface(*GetServices<FOnlineServicesEpicCommon>().GetEOSPlatformHandle());
+	AchievementsHandle = EOS_Platform_GetAchievementsInterface(*GetServices<FOnlineServicesEOSGS>().GetEOSPlatformHandle());
 	check(AchievementsHandle);
 
 	OnAchievementsUnlocked = EOS_RegisterComponentEventHandler(
diff --git a/Engine/Plugins/Online/OnlineServicesEOSGS/Source/Private/Online/ExternalUIEOSGS.cpp b/Engine/Plugins/Online/OnlineServicesEOSGS/Source/Private/Online/ExternalUIEOSGS.cpp
--- a/Engine/Plugins/Online/OnlineServicesEOSGS/Source/Private/Online/ExternalUIEOSGS.cpp
+++ b/Engine/Plugins/Online/OnlineServicesEOSGS/Source/Private/Online/ExternalUIEOSGS.cpp
@@ -9,7 +9,7 @@
 
 namespace UE::Online {
 
-FExternalUIEOSGS::FExternalUIEOSGS(FOnlineServicesEpicCommon& InServices)
+FExternalUIEOSGS::FExternalUIEOSGS(FOnlineServicesEOSGS& InServices)
 	: Super(InServices)
 {
 }
@@ -21,7 +21,7 @@ void FExternalUIEOSGS::Initialize()
 	IEOSSDKManager* SDKManager = IEOSSDKManager::Get();
 	check(SDKManager);
 
-	UIInterfaceHandle = EOS_Platform_GetUIInterface(*GetServices<FOnlineServicesEpicCommon>().GetEOSPlatformHandle());
+	UIInterfaceHandle = EOS_Platform_GetUIInterface(*GetServices<FOnlineServicesEOSGS>().GetEOSPlatformHandle());
 	check(UIInterfaceHandle != nullptr);
 
 	RegisterEventHandlers();
diff --git a/Engine/Plugins/Online/OnlineServicesEOSGS/Source/Private/Online/LeaderboardsEOSGS.cpp b/Engine/Plugins/Online/OnlineServicesEOSGS/Source/Private/Online/LeaderboardsEOSGS.cpp
--- a/Engine/Plugins/Online/OnlineServicesEOSGS/Source/Private/Online/LeaderboardsEOSGS.cpp
+++ b/Engine/Plugins/Online/OnlineServicesEOSGS/Source/Private/Online/LeaderboardsEOSGS.cpp
@@ -16,7 +16,7 @@ namespace UE::Online {
 #define EOS_LEADERBOARD_STAT_NAME_MAX_LENGTH 256 + 1 // 256 plus null terminator
 #define EOS_LEADERBOARD_MAX_ENTRIES 1000
 
-FLeaderboardsEOSGS::FLeaderboardsEOSGS(FOnlineServicesEpicCommon& InServices)
+FLeaderboardsEOSGS::FLeaderboardsEOSGS(FOnlineServicesEOSGS& InServices)
 	: Super(InServices)
 {
 }
@@ -25,7 +25,7 @@ void FLeaderboardsEOSGS::Initialize()
 {
 	Super::Initialize();
 
-	LeaderboardsHandle = EOS_Platform_GetLeaderboardsInterface(*GetServices<FOnlineServicesEpicCommon>().GetEOSPlatformHandle());
+	LeaderboardsHandle = EOS_Platform_GetLeaderboardsInterface(*GetServices<FOnlineServicesEOSGS>().GetEOSPlatformHandle());
 	check(LeaderboardsHandle);
 }
 
diff --git a/Engine/Plugins/Online/OnlineServicesEOSGS/Source/Private/Online/LobbiesEOSGS.cpp b/Engine/Plugins/Online/OnlineServicesEOSGS/Source/Private/Online/LobbiesEOSGS.cpp
--- a/Engine/Plugins/Online/OnlineServicesEOSGS/Source/Private/Online/LobbiesEOSGS.cpp
+++ b/Engine/Plugins/Online/OnlineServicesEOSGS/Source/Private/Online/LobbiesEOSGS.cpp
@@ -24,7 +24,7 @@ namespace UE::Online {
 
 static const int32 MaxAttributeSize = 1000;
 
-FLobbiesEOSGS::FLobbiesEOSGS(FOnlineServicesEpicCommon& InServices)
+FLobbiesEOSGS::FLobbiesEOSGS(FOnlineServicesEOSGS& InServices)
 	: FLobbiesCommon(InServices)
 {
 }
@@ -36,7 +36,7 @@ void FLobbiesEOSGS::Initialize()
 	IEOSSDKManager* SDKManager = IEOSSDKManager::Get();
 	check(SDKManager);
 
-	EOS_HLobby LobbyInterfaceHandle = EOS_Platform_GetLobbyInterface(*GetServices<FOnlineServicesEpicCommon>().GetEOSPlatformHandle());
+	EOS_HLobby LobbyInterfaceHandle = EOS_Platform_GetLobbyInterface(*GetServices<FOnlineServicesEOSGS>().GetEOSPlatformHandle());
 	check(LobbyInterfaceHandle != nullptr);
 
 	LobbyPrerequisites = MakeShared<FLobbyPrerequisitesEOS>(FLobbyPrerequisitesEOS{
diff --git a/Engine/Plugins/Online/OnlineServicesEOSGS/Source/Private/Online/OnlineServicesEOSGS.cpp b/Engine/Plugins/Online/OnlineServicesEOSGS/Source/Private/Online/OnlineServicesEOSGS.cpp
--- a/Engine/Plugins/Online/OnlineServicesEOSGS/Source/Private/Online/OnlineServicesEOSGS.cpp
+++ b/Engine/Plugins/Online/OnlineServicesEOSGS/Source/Private/Online/OnlineServicesEOSGS.cpp
@@ -65,6 +65,27 @@ void FOnlineServicesEOSGS::Destroy()
 #endif
 }
 
+void FOnlineServicesEOSGS::WarnIfEncryptionKeyMissing(const FString& InterfaceName) const
+{
+	if (IEOSSDKManager* Manager = IEOSSDKManager::Get())
+	{
+		const FString& PlatformConfigName = GetEOSPlatformHandle()->GetConfigName();
+		if (const FEOSSDKPlatformConfig* Config = Manager->GetPlatformConfig(PlatformConfigName))
+		{
+			const FString& EncryptionKey = Config->EncryptionKey;
+			if (EncryptionKey.IsEmpty())
+			{
+				UE_LOG(LogOnlineServices, Verbose, TEXT("%s interface not available due to missing ClientEncryptionKey in config."), *InterfaceName);
+			}
+			else
+			{
+				// If we have an encryption key and still can't get the interface, something weird is going on.
+				UE_LOG(LogOnlineServices, Warning, TEXT("%s interface not available despite encryption key being present."), *InterfaceName);
+			}
+		}
+	}
+}
+
 void FOnlineServicesEOSGS::RegisterComponents()
 {
 	Components.Register<FAchievementsEOSGS>(*this);
diff --git a/Engine/Plugins/Online/OnlineServicesEOSGS/Source/Private/Online/OnlineServicesEOSGSModule.cpp b/Engine/Plugins/Online/OnlineServicesEOSGS/Source/Private/Online/OnlineServicesEOSGSModule.cpp
--- a/Engine/Plugins/Online/OnlineServicesEOSGS/Source/Private/Online/OnlineServicesEOSGSModule.cpp
+++ b/Engine/Plugins/Online/OnlineServicesEOSGS/Source/Private/Online/OnlineServicesEOSGSModule.cpp
@@ -52,12 +52,8 @@ void FOnlineServicesEOSGSModule::StartupModule()
 
 	static FOnlineAccountIdRegistryEOSGS AccountIdRegistry(EOnlineServices::Epic);
 	FOnlineIdRegistryRegistry::Get().RegisterAccountIdRegistry(EOnlineServices::Epic, &AccountIdRegistry, GetRegistryPriority());
-
-	static FOnlineSessionIdRegistryEOSGS SessionIdRegistry(EOnlineServices::Epic);
-	FOnlineIdRegistryRegistry::Get().RegisterSessionIdRegistry(EOnlineServices::Epic, &SessionIdRegistry, GetRegistryPriority());
-
-	static FOnlineSessionInviteIdRegistryEOSGS SessionInviteIdRegistry(EOnlineServices::Epic);
-	FOnlineIdRegistryRegistry::Get().RegisterSessionInviteIdRegistry(EOnlineServices::Epic, &SessionInviteIdRegistry, GetRegistryPriority());
+	FOnlineIdRegistryRegistry::Get().RegisterSessionIdRegistry(EOnlineServices::Epic, &FOnlineSessionIdRegistryEOSGS::Get(), GetRegistryPriority());
+	FOnlineIdRegistryRegistry::Get().RegisterSessionInviteIdRegistry(EOnlineServices::Epic, &FOnlineSessionInviteIdRegistryEOSGS::Get(), GetRegistryPriority());
 
 	// Initialize the platform factory on startup.  This is necessary for the SDK to bind to rendering and input very early.
 	FOnlineServicesEOSGSPlatformFactory::Get();
diff --git a/Engine/Plugins/Online/OnlineServicesEOSGS/Source/Private/Online/SessionsEOSGS.cpp b/Engine/Plugins/Online/OnlineServicesEOSGS/Source/Private/Online/SessionsEOSGS.cpp
--- a/Engine/Plugins/Online/OnlineServicesEOSGS/Source/Private/Online/SessionsEOSGS.cpp
+++ b/Engine/Plugins/Online/OnlineServicesEOSGS/Source/Private/Online/SessionsEOSGS.cpp
@@ -17,35 +17,33 @@ namespace UE::Online {
 
 /** FOnlineSessionIdRegistryEOSGS */
 
-FOnlineSessionIdRegistryEOSGS::FOnlineSessionIdRegistryEOSGS(EOnlineServices ServicesType)
-	: FOnlineSessionIdRegistryLAN(ServicesType)
+FOnlineSessionIdRegistryEOSGS::FOnlineSessionIdRegistryEOSGS()
+	: FOnlineSessionIdRegistryLAN(EOnlineServices::Epic)
 {
 }
 
-FOnlineSessionIdRegistryEOSGS& FOnlineSessionIdRegistryEOSGS::GetRegistered(EOnlineServices Services)
+FOnlineSessionIdRegistryEOSGS& FOnlineSessionIdRegistryEOSGS::Get()
 {
-	IOnlineSessionIdRegistry* Registry = FOnlineIdRegistryRegistry::Get().GetSessionIdRegistry(Services);
-	check(Registry);
-	return *static_cast<FOnlineSessionIdRegistryEOSGS*>(Registry);
+	static FOnlineSessionIdRegistryEOSGS Instance;
+	return Instance;
 }
 
 /** FOnlineSessionInviteIdRegistryEOSGS */
 
-FOnlineSessionInviteIdRegistryEOSGS::FOnlineSessionInviteIdRegistryEOSGS(EOnlineServices ServicesType)
-	: FOnlineSessionInviteIdStringRegistry(ServicesType)
+FOnlineSessionInviteIdRegistryEOSGS::FOnlineSessionInviteIdRegistryEOSGS()
+	: FOnlineSessionInviteIdStringRegistry(EOnlineServices::Epic)
 {
 }
 
-FOnlineSessionInviteIdRegistryEOSGS& FOnlineSessionInviteIdRegistryEOSGS::GetRegistered(EOnlineServices Services)
+FOnlineSessionInviteIdRegistryEOSGS& FOnlineSessionInviteIdRegistryEOSGS::Get()
 {
-	IOnlineSessionInviteIdRegistry* Registry = FOnlineIdRegistryRegistry::Get().GetSessionInviteIdRegistry(Services);
-	check(Registry);
-	return *static_cast<FOnlineSessionInviteIdRegistryEOSGS*>(Registry);
+	static FOnlineSessionInviteIdRegistryEOSGS Instance;
+	return Instance;
 }
 
 /** FSessionEOSGS */
 
-FSessionEOSGS::FSessionEOSGS(const FSessionsEOSGS& SessionsInstance, const TSharedPtr<FSessionDetailsHandleEOSGS>& InSessionDetailsHandle)
+FSessionEOSGS::FSessionEOSGS(const TSharedPtr<FSessionDetailsHandleEOSGS>& InSessionDetailsHandle)
 	: SessionDetailsHandle(InSessionDetailsHandle)
 {
 	EOS_SessionDetails_CopyInfoOptions CopyInfoOptions = { };
@@ -56,7 +54,7 @@ FSessionEOSGS::FSessionEOSGS(const FSessionsEOSGS& SessionsInstance, const TShar
 	EOS_EResult CopyInfoResult = EOS_SessionDetails_CopyInfo(SessionDetailsHandle->SessionDetailsHandle, &CopyInfoOptions, &SessionDetailsInfo);
 	if (CopyInfoResult == EOS_EResult::EOS_Success)
 	{
-		SessionInfo.SessionId = SessionsInstance.CreateSessionId(UTF8_TO_TCHAR(SessionDetailsInfo->SessionId));
+		SessionInfo.SessionId = FSessionsEOSGS::CreateSessionId(FString(SessionDetailsInfo->SessionId));
 		SessionSettings.CustomSettings.Emplace(EOSGS_HOST_ADDRESS_ATTRIBUTE_KEY, FCustomSessionSetting{ FSchemaVariant(FString(UTF8_TO_TCHAR(SessionDetailsInfo->HostAddress))), ESchemaAttributeVisibility::Public });
 
 		// bInvitesAllowed and bAllowJoinInProgress should both have the same value, and that value corresponds to bAllowNewMembers
@@ -137,7 +135,7 @@ const FSessionEOSGS& FSessionEOSGS::Cast(const ISession& InSession)
 
 /** FSessionsEOSGS */
 
-FSessionsEOSGS::FSessionsEOSGS(FOnlineServicesEpicCommon& InServices)
+FSessionsEOSGS::FSessionsEOSGS(FOnlineServicesEOSGS& InServices)
 	: Super(InServices)
 {
 }
@@ -146,7 +144,7 @@ void FSessionsEOSGS::Initialize()
 {
 	Super::Initialize();
 
-	SessionsHandle = EOS_Platform_GetSessionsInterface(*GetServices<FOnlineServicesEpicCommon>().GetEOSPlatformHandle());
+	SessionsHandle = EOS_Platform_GetSessionsInterface(*GetServices<FOnlineServicesEOSGS>().GetEOSPlatformHandle());
 	check(SessionsHandle);
 
 	RegisterEventHandlers();
@@ -1007,8 +1005,8 @@ void FSessionsEOSGS::SetSessionSearchSessionId(FSessionSearchHandleEOSGS& Sessio
 	Options.ApiVersion = 1;
 	UE_EOS_CHECK_API_MISMATCH(EOS_SESSIONSEARCH_SETSESSIONID_API_LATEST, 1);
 
-	const auto SessionIdUtf8 = StringCast<UTF8CHAR>(*ToString(SessionId));
-	Options.SessionId = (const char*)SessionIdUtf8.Get();
+	const FTCHARToUTF8 SessionIdUtf8(*FOnlineSessionIdRegistryEOSGS::Get().ToLogString(SessionId));
+	Options.SessionId = SessionIdUtf8.Get();
 
 	EOS_EResult ResultCode = EOS_SessionSearch_SetSessionId(SessionSearchHandle.SearchHandle, &Options);
 	if (ResultCode != EOS_EResult::EOS_Success)
@@ -1051,14 +1049,14 @@ void FSessionsEOSGS::WriteSessionSearchHandle(FSessionSearchHandleEOSGS& Session
 	}
 }
 
-FOnlineSessionId FSessionsEOSGS::CreateSessionId(const FString& SessionId) const
+FOnlineSessionId FSessionsEOSGS::CreateSessionId(const FString& SessionId)
 {
-	return FOnlineSessionIdRegistryEOSGS::GetRegistered(Services.GetServicesProvider()).BasicRegistry.FindOrAddHandle(SessionId);
+	return FOnlineSessionIdRegistryEOSGS::Get().BasicRegistry.FindOrAddHandle(SessionId);
 }
 
 FSessionInviteId FSessionsEOSGS::CreateSessionInviteId(const FString& SessionInviteId) const
 {
-	return FOnlineSessionInviteIdRegistryEOSGS::GetRegistered(Services.GetServicesProvider()).BasicRegistry.FindOrAddHandle(SessionInviteId);
+	return FOnlineSessionInviteIdRegistryEOSGS::Get().BasicRegistry.FindOrAddHandle(SessionInviteId);
 }
 
 TFuture<TOnlineResult<FFindSessions>> FSessionsEOSGS::FindSessionsImpl(const FFindSessions::Params& Params)
@@ -1203,7 +1201,7 @@ TFuture<TOnlineResult<FJoinSession>> FSessionsEOSGS::JoinSessionImpl(const FJoin
 	if (GetSessionByIdResult.IsError())
 	{
 		// If no result is found, the id might be expired, which we should notify
-		if (FOnlineSessionIdRegistryEOSGS::GetRegistered(Services.GetServicesProvider()).IsSessionIdExpired(Params.SessionId))
+		if (FOnlineSessionIdRegistryEOSGS::Get().IsSessionIdExpired(Params.SessionId))
 		{
 			UE_LOG(LogOnlineServices, Warning, TEXT("[%s] SessionId parameter [%s] is expired. Please call FindSessions to get an updated list of available sessions "), UTF8_TO_TCHAR(__FUNCTION__), *ToLogString(Params.SessionId));
 		}
@@ -1251,7 +1249,7 @@ TFuture<TOnlineResult<FJoinSession>> FSessionsEOSGS::JoinSessionImpl(const FJoin
 			if (GetSessionByIdResult.IsError())
 			{
 				// If no result is found, the id might be expired, which we should notify
-				if (FOnlineSessionIdRegistryEOSGS::GetRegistered(Services.GetServicesProvider()).IsSessionIdExpired(Params.SessionId))
+				if (FOnlineSessionIdRegistryEOSGS::Get().IsSessionIdExpired(Params.SessionId))
 				{
 					UE_LOG(LogOnlineServices, Warning, TEXT("[%s] SessionId parameter [%s] is expired. Please call FindSessions to get an updated list of available sessions "), UTF8_TO_TCHAR(__FUNCTION__), *ToLogString(Params.SessionId));
 				}
@@ -1399,7 +1397,7 @@ TOnlineAsyncOpHandle<FBuildSessionFromDetailsHandle> FSessionsEOSGS::BuildSessio
 			if (TOnlineAsyncOpPtr<FBuildSessionFromDetailsHandle> StrongOp = WeakOp.Pin())
 			{
 				// After all the ids are resolved, we can build the session safely
-				StrongOp->SetResult(FBuildSessionFromDetailsHandle::Result{ StrongOp->GetParams().LocalAccountId, MakeShared<FSessionEOSGS>(*this, Params.SessionDetailsHandleEOSGS) });
+				StrongOp->SetResult(FBuildSessionFromDetailsHandle::Result{ StrongOp->GetParams().LocalAccountId, MakeShared<FSessionEOSGS>(Params.SessionDetailsHandleEOSGS) });
 			}
 		});
 	}
@@ -1470,7 +1468,7 @@ TFuture<TOnlineResult<FRejectSessionInvite>> FSessionsEOSGS::RejectSessionInvite
 
 	RejectInviteOptions.LocalUserId = GetProductUserIdChecked(Params.LocalAccountId);
 
-	const FString& InviteIdStr = FOnlineSessionInviteIdRegistryEOSGS::GetRegistered(Services.GetServicesProvider()).BasicRegistry.FindIdValue(Params.SessionInviteId);
+	const FString& InviteIdStr = FOnlineSessionInviteIdRegistryEOSGS::Get().BasicRegistry.FindIdValue(Params.SessionInviteId);
 	const FTCHARToUTF8 InviteIdUtf8(*InviteIdStr);
 	RejectInviteOptions.InviteId = InviteIdUtf8.Get();
 
diff --git a/Engine/Plugins/Online/OnlineServicesEOSGS/Source/Private/Online/StatsEOSGS.cpp b/Engine/Plugins/Online/OnlineServicesEOSGS/Source/Private/Online/StatsEOSGS.cpp
--- a/Engine/Plugins/Online/OnlineServicesEOSGS/Source/Private/Online/StatsEOSGS.cpp
+++ b/Engine/Plugins/Online/OnlineServicesEOSGS/Source/Private/Online/StatsEOSGS.cpp
@@ -21,7 +21,7 @@ struct FStatNameRawBufferEOS
 	char StatName[UE_ONLINE_STAT_EOS_NAME_MAX_LENGTH];
 };
 
-FStatsEOSGS::FStatsEOSGS(FOnlineServicesEpicCommon& InServices)
+FStatsEOSGS::FStatsEOSGS(FOnlineServicesEOSGS& InServices)
 	: Super(InServices)
 {
 }
@@ -30,7 +30,7 @@ void FStatsEOSGS::Initialize()
 {
 	Super::Initialize();
 
-	StatsHandle = EOS_Platform_GetStatsInterface(*GetServices<FOnlineServicesEpicCommon>().GetEOSPlatformHandle());
+	StatsHandle = EOS_Platform_GetStatsInterface(*GetServices<FOnlineServicesEOSGS>().GetEOSPlatformHandle());
 	check(StatsHandle);
 }
 
diff --git a/Engine/Plugins/Online/OnlineServicesEOSGS/Source/Private/Online/TitleFileEOSGS.cpp b/Engine/Plugins/Online/OnlineServicesEOSGS/Source/Private/Online/TitleFileEOSGS.cpp
--- a/Engine/Plugins/Online/OnlineServicesEOSGS/Source/Private/Online/TitleFileEOSGS.cpp
+++ b/Engine/Plugins/Online/OnlineServicesEOSGS/Source/Private/Online/TitleFileEOSGS.cpp
@@ -14,7 +14,7 @@
 
 namespace UE::Online {
 
-FTitleFileEOSGS::FTitleFileEOSGS(FOnlineServicesEpicCommon& InServices)
+FTitleFileEOSGS::FTitleFileEOSGS(FOnlineServicesEOSGS& InServices)
 	: Super(InServices)
 {
 }
@@ -23,7 +23,7 @@ void FTitleFileEOSGS::Initialize()
 {
 	Super::Initialize();
 
-	TitleStorageHandle = EOS_Platform_GetTitleStorageInterface(*GetServices<FOnlineServicesEpicCommon>().GetEOSPlatformHandle());
+	TitleStorageHandle = EOS_Platform_GetTitleStorageInterface(*GetServices<FOnlineServicesEOSGS>().GetEOSPlatformHandle());
 	check(TitleStorageHandle);
 }
 
diff --git a/Engine/Plugins/Online/OnlineServicesEOSGS/Source/Private/Online/UserFileEOSGS.cpp b/Engine/Plugins/Online/OnlineServicesEOSGS/Source/Private/Online/UserFileEOSGS.cpp
--- a/Engine/Plugins/Online/OnlineServicesEOSGS/Source/Private/Online/UserFileEOSGS.cpp
+++ b/Engine/Plugins/Online/OnlineServicesEOSGS/Source/Private/Online/UserFileEOSGS.cpp
@@ -16,7 +16,7 @@ namespace UE::Online {
 
 static const TCHAR* RequestHandleKey = TEXT("RequestHandle");
 
-FUserFileEOSGS::FUserFileEOSGS(FOnlineServicesEpicCommon& InServices)
+FUserFileEOSGS::FUserFileEOSGS(FOnlineServicesEOSGS& InServices)
 	: Super(InServices)
 {
 }
@@ -25,7 +25,7 @@ void FUserFileEOSGS::Initialize()
 {
 	Super::Initialize();
 
-	PlayerDataStorageHandle = EOS_Platform_GetPlayerDataStorageInterface(*GetServices<FOnlineServicesEpicCommon>().GetEOSPlatformHandle());
+	PlayerDataStorageHandle = EOS_Platform_GetPlayerDataStorageInterface(*GetServices<FOnlineServicesEOSGS>().GetEOSPlatformHandle());
 	check(PlayerDataStorageHandle);
 }
 
diff --git a/Engine/Plugins/Online/OnlineServicesEOSGS/Source/Public/Online/AchievementsEOSGS.h b/Engine/Plugins/Online/OnlineServicesEOSGS/Source/Public/Online/AchievementsEOSGS.h
--- a/Engine/Plugins/Online/OnlineServicesEOSGS/Source/Public/Online/AchievementsEOSGS.h
+++ b/Engine/Plugins/Online/OnlineServicesEOSGS/Source/Public/Online/AchievementsEOSGS.h
@@ -2,6 +2,7 @@
 
 #pragma once
 
+#include "CoreMinimal.h"
 #include "Online/AchievementsCommon.h"
 #include "Online/OnlineServicesEOSGSTypes.h"
 
@@ -12,14 +13,14 @@
 
 namespace UE::Online {
 
-class FOnlineServicesEpicCommon;
+class FOnlineServicesEOSGS;
 
 class ONLINESERVICESEOSGS_API FAchievementsEOSGS : public FAchievementsCommon
 {
 public:
 	using Super = FAchievementsCommon;
 
-	FAchievementsEOSGS(FOnlineServicesEpicCommon& InOwningSubsystem);
+	FAchievementsEOSGS(FOnlineServicesEOSGS& InOwningSubsystem);
 	virtual ~FAchievementsEOSGS() = default;
 
 	// IOnlineComponent
diff --git a/Engine/Plugins/Online/OnlineServicesEOSGS/Source/Public/Online/AuthEOSGS.h b/Engine/Plugins/Online/OnlineServicesEOSGS/Source/Public/Online/AuthEOSGS.h
--- a/Engine/Plugins/Online/OnlineServicesEOSGS/Source/Public/Online/AuthEOSGS.h
+++ b/Engine/Plugins/Online/OnlineServicesEOSGS/Source/Public/Online/AuthEOSGS.h
@@ -2,6 +2,7 @@
 
 #pragma once
 
+#include "CoreMinimal.h"
 #include "Online/AuthCommon.h"
 #include "Online/OnlineServicesEOSGSTypes.h"
 
diff --git a/Engine/Plugins/Online/OnlineServicesEOSGS/Source/Public/Online/ExternalUIEOSGS.h b/Engine/Plugins/Online/OnlineServicesEOSGS/Source/Public/Online/ExternalUIEOSGS.h
--- a/Engine/Plugins/Online/OnlineServicesEOSGS/Source/Public/Online/ExternalUIEOSGS.h
+++ b/Engine/Plugins/Online/OnlineServicesEOSGS/Source/Public/Online/ExternalUIEOSGS.h
@@ -2,15 +2,15 @@
 
 #pragma once
 
+#include "CoreMinimal.h"
 #include "Online/ExternalUICommon.h"
-#include "Online/OnlineServicesEOSGSTypes.h"
-#include "Online/OnlineServicesEpicCommon.h"
+#include "OnlineServicesEOSGSTypes.h"
 
 #include "eos_ui_types.h"
 
 namespace UE::Online {
 
-class FOnlineServicesEpicCommon;
+class FOnlineServicesEOSGS;
 
 struct FExternalUIProcessDisplaySettingsUpdatedImp
 {
@@ -37,7 +37,7 @@ class ONLINESERVICESEOSGS_API FExternalUIEOSGS : public FExternalUICommon
 public:
 	using Super = FExternalUICommon;
 
-	FExternalUIEOSGS(FOnlineServicesEpicCommon& InOwningSubsystem);
+	FExternalUIEOSGS(FOnlineServicesEOSGS& InOwningSubsystem);
 
 	virtual void Initialize() override;
 	virtual void PreShutdown() override;
diff --git a/Engine/Plugins/Online/OnlineServicesEOSGS/Source/Public/Online/LeaderboardsEOSGS.h b/Engine/Plugins/Online/OnlineServicesEOSGS/Source/Public/Online/LeaderboardsEOSGS.h
--- a/Engine/Plugins/Online/OnlineServicesEOSGS/Source/Public/Online/LeaderboardsEOSGS.h
+++ b/Engine/Plugins/Online/OnlineServicesEOSGS/Source/Public/Online/LeaderboardsEOSGS.h
@@ -2,6 +2,7 @@
 
 #pragma once
 
+#include "CoreMinimal.h"
 #include "Online/LeaderboardsCommon.h"
 
 #if defined(EOS_PLATFORM_BASE_FILE_NAME)
@@ -11,14 +12,14 @@
 
 namespace UE::Online {
 
-class FOnlineServicesEpicCommon;
+class FOnlineServicesEOSGS;
 
 class ONLINESERVICESEOSGS_API FLeaderboardsEOSGS : public FLeaderboardsCommon
 {
 public:
 	using Super = FLeaderboardsCommon;
 
-	FLeaderboardsEOSGS(FOnlineServicesEpicCommon& InOwningSubsystem);
+	FLeaderboardsEOSGS(FOnlineServicesEOSGS& InOwningSubsystem);
 
 	// TOnlineComponent
 	virtual void Initialize() override;
diff --git a/Engine/Plugins/Online/OnlineServicesEOSGS/Source/Public/Online/LobbiesEOSGS.h b/Engine/Plugins/Online/OnlineServicesEOSGS/Source/Public/Online/LobbiesEOSGS.h
--- a/Engine/Plugins/Online/OnlineServicesEOSGS/Source/Public/Online/LobbiesEOSGS.h
+++ b/Engine/Plugins/Online/OnlineServicesEOSGS/Source/Public/Online/LobbiesEOSGS.h
@@ -2,6 +2,7 @@
 
 #pragma once
 
+#include "CoreMinimal.h"
 #include "Online/LobbiesCommon.h"
 #include "OnlineServicesEOSGSTypes.h"
 
@@ -16,7 +17,7 @@ struct FLobbyPrerequisitesEOS;
 class FLobbyDataEOS;
 class FLobbyDataRegistryEOS;
 class FLobbyInviteDataEOS;
-class FOnlineServicesEpicCommon;
+class FOnlineServicesEOSGS;
 class FLobbySearchEOS;
 struct FClientLobbyDataChanges;
 struct FClientLobbyMemberDataChanges;
@@ -188,7 +189,7 @@ class FLobbiesEOSGS : public FLobbiesCommon
 public:
 	using Super = FLobbiesCommon;
 
-	FLobbiesEOSGS(FOnlineServicesEpicCommon& InServices);
+	FLobbiesEOSGS(FOnlineServicesEOSGS& InServices);
 
 	virtual void Initialize() override;
 	virtual void PreShutdown() override;
diff --git a/Engine/Plugins/Online/OnlineServicesEOSGS/Source/Public/Online/NboSerializerEOSGSSvc.h b/Engine/Plugins/Online/OnlineServicesEOSGS/Source/Public/Online/NboSerializerEOSGSSvc.h
--- a/Engine/Plugins/Online/OnlineServicesEOSGS/Source/Public/Online/NboSerializerEOSGSSvc.h
+++ b/Engine/Plugins/Online/OnlineServicesEOSGS/Source/Public/Online/NboSerializerEOSGSSvc.h
@@ -20,27 +20,16 @@ namespace NboSerializerEOSGSSvc {
 
 /** NboSerializeToBuffer methods */
 
-inline void SerializeToBuffer(FNboSerializeToBuffer& Ar, const EOnlineServices ServicesType)
-{
-	Ar << (__underlying_type(EOnlineServices))ServicesType;
-}
-
 inline void SerializeToBuffer(FNboSerializeToBuffer& Ar, const FAccountId& UniqueId)
 {
-	const TArray<uint8> Data = FOnlineAccountIdRegistryEOSGS::GetRegistered(UniqueId.GetOnlineServicesType()).ToReplicationData(UniqueId);
-
-	SerializeToBuffer(Ar, UniqueId.GetOnlineServicesType());
-
+	TArray<uint8> Data = FOnlineAccountIdRegistryEOSGS::GetRegistered(UniqueId.GetOnlineServicesType()).ToReplicationData(UniqueId);
 	Ar << Data.Num();
 	Ar.WriteBinary(Data.GetData(), Data.Num());
 }
 
 inline void SerializeToBuffer(FNboSerializeToBuffer& Ar, const FOnlineSessionId& SessionId)
 {
-	const TArray<uint8> Data = FOnlineSessionIdRegistryEOSGS::GetRegistered(SessionId.GetOnlineServicesType()).ToReplicationData(SessionId);
-
-	SerializeToBuffer(Ar, SessionId.GetOnlineServicesType());
-
+	TArray<uint8> Data = FOnlineSessionIdRegistryEOSGS::Get().ToReplicationData(SessionId);
 	Ar << Data.Num();
 	Ar.WriteBinary(Data.GetData(), Data.Num());
 }
@@ -48,7 +37,7 @@ inline void SerializeToBuffer(FNboSerializeToBuffer& Ar, const FOnlineSessionId&
 inline void SerializeToBuffer(FNboSerializeToBuffer& Packet, const FSessionMemberIdsSet& SessionMembersSet)
 {
 	Packet << SessionMembersSet.Num();
-	
+
 	for (const FAccountId& SessionMember : SessionMembersSet)
 	{
 		SerializeToBuffer(Packet, SessionMember);
@@ -57,41 +46,22 @@ inline void SerializeToBuffer(FNboSerializeToBuffer& Packet, const FSessionMembe
 
 /** NboSerializeFromBuffer methods */
 
-inline void SerializeFromBuffer(FNboSerializeFromBuffer& Ar, EOnlineServices& ServicesType)
-{
-	using UnderlyingType = __underlying_type(EOnlineServices);
-	UnderlyingType Value;
-	Ar >> Value;
-
-	ServicesType = (EOnlineServices)Value;
-}
-
 inline void SerializeFromBuffer(FNboSerializeFromBuffer& Ar, FAccountId& UniqueId)
 {
-	EOnlineServices ServicesType;
-	SerializeFromBuffer(Ar, ServicesType);
-
+	TArray<uint8> Data;
 	int32 Size;
 	Ar >> Size;
-
-	TArray<uint8> Data;
 	Ar.ReadBinaryArray(Data, Size);
-
-	UniqueId = FOnlineAccountIdRegistryEOSGS::GetRegistered(ServicesType).FromReplicationData(Data);
+	UniqueId = FOnlineAccountIdRegistryEOSGS::GetRegistered(UniqueId.GetOnlineServicesType()).FromReplicationData(Data);
 }
 
 inline void SerializeFromBuffer(FNboSerializeFromBuffer& Ar, FOnlineSessionId& SessionId)
 {
-	EOnlineServices ServicesType;
-	SerializeFromBuffer(Ar, ServicesType);
-
+	TArray<uint8> Data;
 	int32 Size;
 	Ar >> Size;
-
-	TArray<uint8> Data;
 	Ar.ReadBinaryArray(Data, Size);
-
-	SessionId = FOnlineSessionIdRegistryEOSGS::GetRegistered(ServicesType).FromReplicationData(Data);
+	SessionId = FOnlineSessionIdRegistryEOSGS::Get().FromReplicationData(Data);
 }
 
 inline void SerializeFromBuffer(FNboSerializeFromBuffer& Packet, FSessionMemberIdsSet& SessionMembersSet)
diff --git a/Engine/Plugins/Online/OnlineServicesEOSGS/Source/Public/Online/OnlineServicesEOSGS.h b/Engine/Plugins/Online/OnlineServicesEOSGS/Source/Public/Online/OnlineServicesEOSGS.h
--- a/Engine/Plugins/Online/OnlineServicesEOSGS/Source/Public/Online/OnlineServicesEOSGS.h
+++ b/Engine/Plugins/Online/OnlineServicesEOSGS/Source/Public/Online/OnlineServicesEOSGS.h
@@ -42,6 +42,7 @@ public:
 	static const TCHAR* GetServiceConfigNameStatic() { return TEXT("EOS"); }
 
 protected:
+	void WarnIfEncryptionKeyMissing(const FString& InterfaceName) const;
 
 #if WITH_ENGINE
 	TSharedPtr<FSocketSubsystemEOS, ESPMode::ThreadSafe> SocketSubsystem;
diff --git a/Engine/Plugins/Online/OnlineServicesEOSGS/Source/Public/Online/SessionsEOSGS.h b/Engine/Plugins/Online/OnlineServicesEOSGS/Source/Public/Online/SessionsEOSGS.h
--- a/Engine/Plugins/Online/OnlineServicesEOSGS/Source/Public/Online/SessionsEOSGS.h
+++ b/Engine/Plugins/Online/OnlineServicesEOSGS/Source/Public/Online/SessionsEOSGS.h
@@ -14,35 +14,23 @@
 
 namespace UE::Online {
 
-class FOnlineServicesEpicCommon;
-class FSessionsEOSGS;
+class FOnlineServicesEOSGS;
 
 class FOnlineSessionIdRegistryEOSGS : public FOnlineSessionIdRegistryLAN
 {
 public:
-	UE_DEPRECATED(5.6, "This method is deprecated, please use GetRegistered(EOnlineServices)")
-	static FOnlineSessionIdRegistryEOSGS& Get() { return GetRegistered(EOnlineServices::Epic); }
-	
-	static FOnlineSessionIdRegistryEOSGS& GetRegistered(EOnlineServices Services);
-	
+	static FOnlineSessionIdRegistryEOSGS& Get();
 private:
-	friend class FOnlineServicesEOSGSModule;
-	friend class FOnlineServicesEpicGameModule;
-	FOnlineSessionIdRegistryEOSGS(EOnlineServices ServicesType);
+	FOnlineSessionIdRegistryEOSGS();
 };
 
 class FOnlineSessionInviteIdRegistryEOSGS : public FOnlineSessionInviteIdStringRegistry
 {
 public:
-	UE_DEPRECATED(5.6, "This method is deprecated, please use GetRegistered(EOnlineServices)")
-	static FOnlineSessionInviteIdRegistryEOSGS& Get() { return GetRegistered(EOnlineServices::Epic); }
-
-	static FOnlineSessionInviteIdRegistryEOSGS& GetRegistered(EOnlineServices Services);
+	static FOnlineSessionInviteIdRegistryEOSGS& Get();
 
 private:
-	friend class FOnlineServicesEOSGSModule;
-	friend class FOnlineServicesEpicGameModule;
-	FOnlineSessionInviteIdRegistryEOSGS(EOnlineServices ServicesType);
+	FOnlineSessionInviteIdRegistryEOSGS();
 };
 
 struct FSessionModificationHandleEOSGS : FNoncopyable
@@ -97,7 +85,7 @@ public:
 	FSessionEOSGS(const FSessionEOSGS& InSession) = default;
 
 	/** This constructor should only be used by BuildSessionFromDetailsHandle, after all user ids in the session have been resolved. */
-	FSessionEOSGS(const FSessionsEOSGS& SessionsInstance, const TSharedPtr<FSessionDetailsHandleEOSGS>& InSessionDetailsHandle);
+	FSessionEOSGS(const TSharedPtr<FSessionDetailsHandleEOSGS>& InSessionDetailsHandle);
 
 	static const FSessionEOSGS& Cast(const ISession& InSession);
 
@@ -194,7 +182,7 @@ public:
 
 	using Super = FSessionsLAN;
 
-	FSessionsEOSGS(FOnlineServicesEpicCommon& InOwningSubsystem);
+	FSessionsEOSGS(FOnlineServicesEOSGS& InOwningSubsystem);
 	virtual ~FSessionsEOSGS() = default;
 
 	// TOnlineComponent
@@ -268,7 +256,7 @@ protected:
 	 */
 	TFuture<TDefaultErrorResult<FSendSingleSessionInviteImpl>> SendSingleSessionInviteImpl(FSendSingleSessionInviteImpl::Params&& Params);
 
-	FOnlineSessionId CreateSessionId(const FString& SessionId) const;
+	static FOnlineSessionId CreateSessionId(const FString& SessionId);
 	FSessionInviteId CreateSessionInviteId(const FString& SessionInviteId) const;
 
 	/**
diff --git a/Engine/Plugins/Online/OnlineServicesEOSGS/Source/Public/Online/StatsEOSGS.h b/Engine/Plugins/Online/OnlineServicesEOSGS/Source/Public/Online/StatsEOSGS.h
--- a/Engine/Plugins/Online/OnlineServicesEOSGS/Source/Public/Online/StatsEOSGS.h
+++ b/Engine/Plugins/Online/OnlineServicesEOSGS/Source/Public/Online/StatsEOSGS.h
@@ -2,6 +2,7 @@
 
 #pragma once
 
+#include "CoreMinimal.h"
 #include "Online/StatsCommon.h"
 
 #if defined(EOS_PLATFORM_BASE_FILE_NAME)
@@ -11,7 +12,7 @@
 
 namespace UE::Online {
 
-class FOnlineServicesEpicCommon;
+class FOnlineServicesEOSGS;
 
 /**
  * Because the value type is int32 in EOS, stats types are converted to int32. So these are limitations: 
@@ -24,7 +25,7 @@ class ONLINESERVICESEOSGS_API FStatsEOSGS : public FStatsCommon
 public:
 	using Super = FStatsCommon;
 
-	FStatsEOSGS(FOnlineServicesEpicCommon& InOwningSubsystem);
+	FStatsEOSGS(FOnlineServicesEOSGS& InOwningSubsystem);
 
 	// TOnlineComponent
 	virtual void Initialize() override;
diff --git a/Engine/Plugins/Online/OnlineServicesEOSGS/Source/Public/Online/TitleFileEOSGS.h b/Engine/Plugins/Online/OnlineServicesEOSGS/Source/Public/Online/TitleFileEOSGS.h
--- a/Engine/Plugins/Online/OnlineServicesEOSGS/Source/Public/Online/TitleFileEOSGS.h
+++ b/Engine/Plugins/Online/OnlineServicesEOSGS/Source/Public/Online/TitleFileEOSGS.h
@@ -2,6 +2,7 @@
 
 #pragma once
 
+#include "CoreMinimal.h"
 #include "Online/TitleFileCommon.h"
 
 #if defined(EOS_PLATFORM_BASE_FILE_NAME)
@@ -11,7 +12,7 @@
 
 namespace UE::Online {
 
-class FOnlineServicesEpicCommon;
+class FOnlineServicesEOSGS;
 
 struct FTitleFileEOSGSConfig
 {
@@ -33,7 +34,7 @@ class ONLINESERVICESEOSGS_API FTitleFileEOSGS : public FTitleFileCommon
 public:
 	using Super = FTitleFileCommon;
 
-	FTitleFileEOSGS(FOnlineServicesEpicCommon& InOwningSubsystem);
+	FTitleFileEOSGS(FOnlineServicesEOSGS& InOwningSubsystem);
 	virtual ~FTitleFileEOSGS() = default;
 
 	// IOnlineComponent
diff --git a/Engine/Plugins/Online/OnlineServicesEOSGS/Source/Public/Online/UserFileEOSGS.h b/Engine/Plugins/Online/OnlineServicesEOSGS/Source/Public/Online/UserFileEOSGS.h
--- a/Engine/Plugins/Online/OnlineServicesEOSGS/Source/Public/Online/UserFileEOSGS.h
+++ b/Engine/Plugins/Online/OnlineServicesEOSGS/Source/Public/Online/UserFileEOSGS.h
@@ -2,6 +2,7 @@
 
 #pragma once
 
+#include "CoreMinimal.h"
 #include "Online/UserFileCommon.h"
 
 #if defined(EOS_PLATFORM_BASE_FILE_NAME)
@@ -11,7 +12,7 @@
 
 namespace UE::Online {
 
-class FOnlineServicesEpicCommon;
+class FOnlineServicesEOSGS;
 
 struct FUserFileEOSGSConfig
 {
@@ -31,7 +32,7 @@ class ONLINESERVICESEOSGS_API FUserFileEOSGS : public FUserFileCommon
 public:
 	using Super = FUserFileCommon;
 
-	FUserFileEOSGS(FOnlineServicesEpicCommon& InOwningSubsystem);
+	FUserFileEOSGS(FOnlineServicesEOSGS& InOwningSubsystem);
 	virtual ~FUserFileEOSGS() = default;
 
 	// IOnlineComponent
diff --git a/Engine/Plugins/Online/OnlineServicesEpicCommon/Source/OnlineServicesEpicCommon/Private/Online/OnlineServicesEpicCommon.cpp b/Engine/Plugins/Online/OnlineServicesEpicCommon/Source/OnlineServicesEpicCommon/Private/Online/OnlineServicesEpicCommon.cpp
--- a/Engine/Plugins/Online/OnlineServicesEpicCommon/Source/OnlineServicesEpicCommon/Private/Online/OnlineServicesEpicCommon.cpp
+++ b/Engine/Plugins/Online/OnlineServicesEpicCommon/Source/OnlineServicesEpicCommon/Private/Online/OnlineServicesEpicCommon.cpp
@@ -69,27 +69,6 @@ void FOnlineServicesEpicCommon::RemoveEOSSDKFastTick(FOnlineAsyncOp& InAsyncOp)
 	}
 }
 
-void FOnlineServicesEpicCommon::WarnIfEncryptionKeyMissing(const FString& InterfaceName) const
-{
-	if (IEOSSDKManager* Manager = IEOSSDKManager::Get())
-	{
-		const FString& PlatformConfigName = GetEOSPlatformHandle()->GetConfigName();
-		if (const FEOSSDKPlatformConfig* Config = Manager->GetPlatformConfig(PlatformConfigName))
-		{
-			const FString& EncryptionKey = Config->EncryptionKey;
-			if (EncryptionKey.IsEmpty())
-			{
-				UE_LOG(LogOnlineServices, Verbose, TEXT("%s interface not available due to missing ClientEncryptionKey in config."), *InterfaceName);
-			}
-			else
-			{
-				// If we have an encryption key and still can't get the interface, something weird is going on.
-				UE_LOG(LogOnlineServices, Warning, TEXT("%s interface not available despite encryption key being present."), *InterfaceName);
-			}
-		}
-	}
-}
-
 void FOnlineServicesEpicCommon::FlushTick(float DeltaSeconds)
 {
 	Super::FlushTick(DeltaSeconds);
diff --git a/Engine/Plugins/Online/OnlineServicesEpicCommon/Source/OnlineServicesEpicCommon/Public/Online/OnlineServicesEpicCommon.h b/Engine/Plugins/Online/OnlineServicesEpicCommon/Source/OnlineServicesEpicCommon/Public/Online/OnlineServicesEpicCommon.h
--- a/Engine/Plugins/Online/OnlineServicesEpicCommon/Source/OnlineServicesEpicCommon/Public/Online/OnlineServicesEpicCommon.h
+++ b/Engine/Plugins/Online/OnlineServicesEpicCommon/Source/OnlineServicesEpicCommon/Public/Online/OnlineServicesEpicCommon.h
@@ -27,8 +27,6 @@ public:
 	void RemoveEOSSDKFastTick(FOnlineAsyncOp& InAsyncOp);
 
 protected:
-	void WarnIfEncryptionKeyMissing(const FString& InterfaceName) const;
-	
 	virtual void FlushTick(float DeltaSeconds) override;
 
 	IEOSPlatformHandlePtr EOSPlatformHandle;
