- Fix not setting error in the case EOS_Connect_Login fails when not using EAS

#rnx

diff --git a/Engine/Plugins/Online/OnlineServicesEOSGS/Source/Private/Online/AuthEOSGS.cpp b/Engine/Plugins/Online/OnlineServicesEOSGS/Source/Private/Online/AuthEOSGS.cpp
--- a/Engine/Plugins/Online/OnlineServicesEOSGS/Source/Private/Online/AuthEOSGS.cpp
+++ b/Engine/Plugins/Online/OnlineServicesEOSGS/Source/Private/Online/AuthEOSGS.cpp
@@ -615,6 +615,7 @@ TOnlineAsyncOpHandle<FAuthLogin> FAuthEOSGS::Login(FAuthLogin::Params&& Params)
 				}
 				else
 				{
+					AsyncOp->SetError(MoveTemp(LoginResult.GetErrorValue()));
 					Promise.EmplaceValue();
 				}
 			}
