Add EOS_Connect_Logout to OnlineServicesEOS

diff --git a/Engine/Plugins/Online/OnlineServicesEOS/Source/Private/Online/AuthEOS.cpp b/Engine/Plugins/Online/OnlineServicesEOS/Source/Private/Online/AuthEOS.cpp
--- a/Engine/Plugins/Online/OnlineServicesEOS/Source/Private/Online/AuthEOS.cpp
+++ b/Engine/Plugins/Online/OnlineServicesEOS/Source/Private/Online/AuthEOS.cpp
@@ -250,10 +250,21 @@ TOnlineAsyncOpHandle<FAuthLogin> FAuthEOS::Login(FAuthLogin::Params&& Params)
 			TPromise<void> Promise;
 			TFuture<void> Future = Promise.GetFuture();
 			LogoutEASImpl(FAuthLogoutEASImpl::Params{ AccountInfoEOS->EpicAccountId })
-			.Next([AsyncOp = InAsyncOp.AsShared(), Error = MoveTemp(CopyUserInfoError), Promise = MoveTemp(Promise)](TDefaultErrorResult<FAuthLogoutEASImpl>&&) mutable -> void
+			.Next([this, AccountInfoEOS, WeakOp = InAsyncOp.AsWeak(), Error = MoveTemp(CopyUserInfoError), Promise = MoveTemp(Promise)](TDefaultErrorResult<FAuthLogoutEASImpl>&&) mutable -> void
 			{
-				AsyncOp->SetError(MoveTemp(Error));
-				Promise.EmplaceValue();
+				if (TSharedPtr<TOnlineAsyncOp<FAuthLogin>> AsyncOp = WeakOp.Pin())
+				{
+					LogoutConnectImpl(FAuthLogoutConnectImpl::Params{ AccountInfoEOS->ProductUserId })
+					.Next([AsyncOp, Error = MoveTemp(Error), Promise = MoveTemp(Promise)](TDefaultErrorResult<FAuthLogoutConnectImpl>&&) mutable -> void
+					{
+						AsyncOp->SetError(MoveTemp(Error));
+						Promise.EmplaceValue();
+					});
+				}
+				else
+				{
+					Promise.EmplaceValue();
+				}
 			});
 
 			return Future;
@@ -508,10 +519,21 @@ TOnlineAsyncOpHandle<FAuthLinkAccount> FAuthEOS::LinkAccount(FAuthLinkAccount::P
 				TPromise<void> Promise;
 				TFuture<void> Future = Promise.GetFuture();
 				LogoutEASImpl(FAuthLogoutEASImpl::Params{ AccountInfoEOS->EpicAccountId })
-				.Next([AsyncOp = InAsyncOp.AsShared(), Error = MoveTemp(CopyUserInfoError), Promise = MoveTemp(Promise)](TDefaultErrorResult<FAuthLogoutEASImpl>&&) mutable -> void
+				.Next([this, AccountInfoEOS, WeakOp = InAsyncOp.AsWeak(), Error = MoveTemp(CopyUserInfoError), Promise = MoveTemp(Promise)](TDefaultErrorResult<FAuthLogoutEASImpl>&&) mutable -> void
 				{
-					AsyncOp->SetError(MoveTemp(Error));
-					Promise.EmplaceValue();
+					if (TSharedPtr<TOnlineAsyncOp<FAuthLinkAccount>> AsyncOp = WeakOp.Pin())
+					{
+						LogoutConnectImpl(FAuthLogoutConnectImpl::Params{ AccountInfoEOS->ProductUserId })
+						.Next([AsyncOp, Error = MoveTemp(Error), Promise = MoveTemp(Promise)](TDefaultErrorResult<FAuthLogoutConnectImpl>&&) mutable -> void
+						{
+							AsyncOp->SetError(MoveTemp(Error));
+							Promise.EmplaceValue();
+						});
+					}
+					else
+					{
+						Promise.EmplaceValue();
+					}
 				});
 
 				return Future;
diff --git a/Engine/Plugins/Online/OnlineServicesEOSGS/Source/Private/Online/AuthEOSGS.cpp b/Engine/Plugins/Online/OnlineServicesEOSGS/Source/Private/Online/AuthEOSGS.cpp
--- a/Engine/Plugins/Online/OnlineServicesEOSGS/Source/Private/Online/AuthEOSGS.cpp
+++ b/Engine/Plugins/Online/OnlineServicesEOSGS/Source/Private/Online/AuthEOSGS.cpp
@@ -769,6 +769,27 @@ TOnlineAsyncOpHandle<FAuthLogout> FAuthEOSGS::Logout(FAuthLogout::Params&& Param
 			return MakeFulfilledPromise<void>().GetFuture();
 		}
 	})
+	// Step 4: Logout of Connect.
+	.Then([this](TOnlineAsyncOp<FAuthLogout>& InAsyncOp)
+	{
+		const TSharedRef<FAccountInfoEOS>& AccountInfoEOS = GetOpDataChecked<TSharedRef<FAccountInfoEOS>>(InAsyncOp, AccountInfoKeyName);
+
+		if (AccountInfoEOS->ProductUserId)
+		{
+			TPromise<void> Promise;
+			TFuture<void> Future = Promise.GetFuture();
+			LogoutConnectImpl(FAuthLogoutConnectImpl::Params{ AccountInfoEOS->ProductUserId })
+			.Next([Promise = MoveTemp(Promise)](TDefaultErrorResult<FAuthLogoutConnectImpl>&&) mutable -> void
+			{
+				Promise.EmplaceValue();
+			});
+			return Future;
+		}
+		else
+		{
+			return MakeFulfilledPromise<void>().GetFuture();
+		}
+	})
 	.Then([this](TOnlineAsyncOp<FAuthLogout>& InAsyncOp)
 	{
 		const TSharedRef<FAccountInfoEOS>& AccountInfoEOS = GetOpDataChecked<TSharedRef<FAccountInfoEOS>>(InAsyncOp, AccountInfoKeyName);
@@ -1140,6 +1161,34 @@ TFuture<TDefaultErrorResult<FAuthLoginConnectImpl>> FAuthEOSGS::LoginConnectImpl
 	return Future;
 }
 
+ONLINESERVICESEOSGS_API TFuture<TDefaultErrorResult<FAuthLogoutConnectImpl>> FAuthEOSGS::LogoutConnectImpl(const FAuthLogoutConnectImpl::Params& LogoutParams)
+{
+	TPromise<TDefaultErrorResult<FAuthLogoutConnectImpl>> Promise;
+	TFuture<TDefaultErrorResult<FAuthLogoutConnectImpl>> Future = Promise.GetFuture();
+
+	EOS_Connect_LogoutOptions LogoutOptions = {};
+	LogoutOptions.ApiVersion = 1; 
+	LogoutOptions.LocalUserId = LogoutParams.ProductUserId;
+	UE_EOS_CHECK_API_MISMATCH(EOS_CONNECT_LOGOUT_API_LATEST, 1);
+
+	EOS_Async(EOS_Connect_Logout, ConnectHandle, LogoutOptions,
+		[Promise = MoveTemp(Promise)](const EOS_Connect_LogoutCallbackInfo* Data) mutable -> void
+		{
+			if (Data->ResultCode == EOS_EResult::EOS_Success)
+			{
+				UE_LOG(LogOnlineServices, Verbose, TEXT("[FAuthEOSGS::LogoutConnectImpl] EOS_Connect_Logout Result: [%s]"), *LexToString(Data->ResultCode));
+				Promise.SetValue(TDefaultErrorResult<FAuthLogoutConnectImpl>(FAuthLogoutConnectImpl::Result{}));
+			}
+			else
+			{
+				UE_LOG(LogOnlineServices, Warning, TEXT("[FAuthEOSGS::LogoutConnectImpl] EOS_Connect_Logout Result: [%s]"), *LexToString(Data->ResultCode));
+				Promise.SetValue(TDefaultErrorResult<FAuthLogoutConnectImpl>(Errors::FromEOSResult(Data->ResultCode)));
+			}
+		});
+
+	return Future;
+}
+
 TOnlineAsyncOpHandle<FAuthConnectLoginRecoveryImpl> FAuthEOSGS::ConnectLoginRecoveryImplOp(FAuthConnectLoginRecoveryImpl::Params&& Params)
 {
 	TOnlineAsyncOpRef<FAuthConnectLoginRecoveryImpl> Op = GetOp<FAuthConnectLoginRecoveryImpl>(MoveTemp(Params));
diff --git a/Engine/Plugins/Online/OnlineServicesEOSGS/Source/Public/Online/AuthEOSGS.h b/Engine/Plugins/Online/OnlineServicesEOSGS/Source/Public/Online/AuthEOSGS.h
--- a/Engine/Plugins/Online/OnlineServicesEOSGS/Source/Public/Online/AuthEOSGS.h
+++ b/Engine/Plugins/Online/OnlineServicesEOSGS/Source/Public/Online/AuthEOSGS.h
@@ -55,6 +55,20 @@ struct FAuthLogoutEASImpl
 	};
 };
 
+struct FAuthLogoutConnectImpl
+{
+	static constexpr TCHAR Name[] = TEXT("LogoutConnectImpl");
+
+	struct Params
+	{
+		EOS_ProductUserId ProductUserId = nullptr;
+	};
+
+	struct Result
+	{
+	};
+};
+
 struct FAuthGetExternalAuthTokenImpl
 {
 	static constexpr TCHAR Name[] = TEXT("GetExternalAuthTokenImpl");
@@ -222,6 +236,7 @@ protected:
 	ONLINESERVICESEOSGS_API TFuture<TDefaultErrorResult<FAuthLogoutEASImpl>> LogoutEASImpl(const FAuthLogoutEASImpl::Params& LogoutParams);
 	ONLINESERVICESEOSGS_API TDefaultErrorResult<FAuthGetExternalAuthTokenImpl> GetExternalAuthTokenImpl(const FAuthGetExternalAuthTokenImpl::Params& Params);
 	ONLINESERVICESEOSGS_API TFuture<TDefaultErrorResult<FAuthLoginConnectImpl>> LoginConnectImpl(const FAuthLoginConnectImpl::Params& LoginParams);
+	ONLINESERVICESEOSGS_API TFuture<TDefaultErrorResult<FAuthLogoutConnectImpl>> LogoutConnectImpl(const FAuthLogoutConnectImpl::Params& LogoutParams);
 	ONLINESERVICESEOSGS_API TOnlineAsyncOpHandle<FAuthConnectLoginRecoveryImpl> ConnectLoginRecoveryImplOp(FAuthConnectLoginRecoveryImpl::Params&& Params);
 	ONLINESERVICESEOSGS_API TOnlineAsyncOpHandle<FAuthHandleConnectLoginStatusChangedImpl> HandleConnectLoginStatusChangedImplOp(FAuthHandleConnectLoginStatusChangedImpl::Params&& Params);
 	ONLINESERVICESEOSGS_API TOnlineAsyncOpHandle<FAuthHandleConnectAuthNotifyExpirationImpl> HandleConnectAuthNotifyExpirationImplOp(FAuthHandleConnectAuthNotifyExpirationImpl::Params&& Params);
@@ -299,6 +314,13 @@ END_ONLINE_STRUCT_META()
 BEGIN_ONLINE_STRUCT_META(FAuthLogoutEASImpl::Result)
 END_ONLINE_STRUCT_META()
 
+BEGIN_ONLINE_STRUCT_META(FAuthLogoutConnectImpl::Params)
+ONLINE_STRUCT_FIELD(FAuthLogoutConnectImpl::Params, ProductUserId)
+END_ONLINE_STRUCT_META()
+
+BEGIN_ONLINE_STRUCT_META(FAuthLogoutConnectImpl::Result)
+END_ONLINE_STRUCT_META()
+
 BEGIN_ONLINE_STRUCT_META(FAuthGetExternalAuthTokenImpl::Params)
 	ONLINE_STRUCT_FIELD(FAuthGetExternalAuthTokenImpl::Params, EpicAccountId)
 END_ONLINE_STRUCT_META()
