- Fix up a few places where static_cast was used instead of the template param on GetServices

#rnx
#DisabledValidations LowLevelTestsOnlineTests, LowLevelTestsFNOnlineFrameworkTests, LowLevelTestsOSMCPTests

diff --git a/Engine/Plugins/Online/OnlineServicesEOSGS/Source/Private/Online/PlayerReportsEOSGS.cpp b/Engine/Plugins/Online/OnlineServicesEOSGS/Source/Private/Online/PlayerReportsEOSGS.cpp
--- a/Engine/Plugins/Online/OnlineServicesEOSGS/Source/Private/Online/PlayerReportsEOSGS.cpp
+++ b/Engine/Plugins/Online/OnlineServicesEOSGS/Source/Private/Online/PlayerReportsEOSGS.cpp
@@ -21,7 +21,7 @@ void FPlayerReportsEOSGS::Initialize()
 {
 	Super::Initialize(); 
 
-	PlayerReportsHandle = EOS_Platform_GetReportsInterface(*static_cast<FOnlineServicesEOSGS&>(GetServices()).GetEOSPlatformHandle());
+	PlayerReportsHandle = EOS_Platform_GetReportsInterface(*GetServices<FOnlineServicesEpicCommon>().GetEOSPlatformHandle());
 	check(PlayerReportsHandle);
 
 	// No need to register for any EOS notifications/events
diff --git a/Engine/Plugins/Online/OnlineServicesEOSGS/Source/Private/Online/PlayerSanctionsEOSGS.cpp b/Engine/Plugins/Online/OnlineServicesEOSGS/Source/Private/Online/PlayerSanctionsEOSGS.cpp
--- a/Engine/Plugins/Online/OnlineServicesEOSGS/Source/Private/Online/PlayerSanctionsEOSGS.cpp
+++ b/Engine/Plugins/Online/OnlineServicesEOSGS/Source/Private/Online/PlayerSanctionsEOSGS.cpp
@@ -21,7 +21,7 @@ void FPlayerSanctionsEOSGS::Initialize()
 {
 	Super::Initialize(); 
 
-	PlayerSanctionsHandle = EOS_Platform_GetSanctionsInterface(*static_cast<FOnlineServicesEOSGS&>(GetServices()).GetEOSPlatformHandle());
+	PlayerSanctionsHandle = EOS_Platform_GetSanctionsInterface(*GetServices<FOnlineServicesEpicCommon>().GetEOSPlatformHandle());
 	check(PlayerSanctionsHandle);
 
 	// No need to register for any EOS notifications/events
