[EOS] Change login recover to use weak lambda to prevent crash

#tests Running pie

diff --git a/Engine/Plugins/Online/OnlineServicesEOSGS/Source/Private/Online/AuthEOSGS.cpp b/Engine/Plugins/Online/OnlineServicesEOSGS/Source/Private/Online/AuthEOSGS.cpp
--- a/Engine/Plugins/Online/OnlineServicesEOSGS/Source/Private/Online/AuthEOSGS.cpp
+++ b/Engine/Plugins/Online/OnlineServicesEOSGS/Source/Private/Online/AuthEOSGS.cpp
@@ -1510,9 +1510,12 @@ void FAuthEOSGS::InitializeConnectLoginRecoveryTimer(const TSharedRef<FAccountIn
 	}
 
 	AccountInfoEOS->RestoreLoginTimer = FTSBackgroundableTicker::GetCoreTicker().AddTicker(
-		FTickerDelegate::CreateLambda([this, EpicAccountId = AccountInfoEOS->EpicAccountId](float) -> bool
+		FTickerDelegate::CreateLambda([this, WeakThis = AsWeak(), EpicAccountId = AccountInfoEOS->EpicAccountId](float) -> bool
 		{
-			ConnectLoginRecoveryImplOp(FAuthConnectLoginRecoveryImpl::Params{ EpicAccountId });
+			if (WeakThis.IsValid())
+			{
+				ConnectLoginRecoveryImplOp(FAuthConnectLoginRecoveryImpl::Params { EpicAccountId });
+			}
 
 			// One-shot timer.
 			return false;
