- Remove unused constants

#DisabledValidations LowLevelTestsOnlineTests, LowLevelTestsFNOnlineFrameworkTests, lowLevelTestsOSMCPTests

diff --git a/Engine/Plugins/Online/OnlineServicesEOSGS/Source/Public/Online/AuthEOSGS.h b/Engine/Plugins/Online/OnlineServicesEOSGS/Source/Public/Online/AuthEOSGS.h
--- a/Engine/Plugins/Online/OnlineServicesEOSGS/Source/Public/Online/AuthEOSGS.h
+++ b/Engine/Plugins/Online/OnlineServicesEOSGS/Source/Public/Online/AuthEOSGS.h
@@ -13,14 +13,10 @@
 
 namespace UE::Online {
 
-const int EOS_STRING_BUFFER_LENGTH = 256;
-const int EOS_MAX_TOKEN_SIZE = 4096;
-
 class FOnlineServicesEOSGS;
 struct FAccountInfoEOS;
 class FAccountInfoRegistryEOS;
 
-
 struct FAuthLoginEASImpl
 {
 	static constexpr TCHAR Name[] = TEXT("LoginEASImpl");
