- Remove unnecessary resolves of local account id

diff --git a/Engine/Plugins/Online/OnlineServicesEOSGS/Source/Private/Online/SessionsEOSGS.cpp b/Engine/Plugins/Online/OnlineServicesEOSGS/Source/Private/Online/SessionsEOSGS.cpp
--- a/Engine/Plugins/Online/OnlineServicesEOSGS/Source/Private/Online/SessionsEOSGS.cpp
+++ b/Engine/Plugins/Online/OnlineServicesEOSGS/Source/Private/Online/SessionsEOSGS.cpp
@@ -1510,41 +1510,30 @@ TFuture<TOnlineResult<FAddSessionMember>> FSessionsEOSGS::AddSessionMemberImpl(c
 	TPromise<TOnlineResult<FAddSessionMember>> Promise;
 	TFuture<TOnlineResult<FAddSessionMember>> Future = Promise.GetFuture();
 
-	Services.Get<IEpicProductUserIdResolver>()->ResolveAccountIds(Params.LocalAccountId, { GetProductUserIdChecked(Params.LocalAccountId) })
-		.Next([this, Promise = MoveTemp(Promise), Params](TArray<FAccountId> ResolvedAccountIds) mutable
-			{
-				if (!ResolvedAccountIds.IsEmpty())
-				{
-					EOS_Sessions_RegisterPlayersOptions Options = { };
-					Options.ApiVersion = 3;
-					UE_EOS_CHECK_API_MISMATCH(EOS_SESSIONS_REGISTERPLAYERS_API_LATEST, 3);
+	EOS_Sessions_RegisterPlayersOptions Options = { };
+	Options.ApiVersion = 3;
+	UE_EOS_CHECK_API_MISMATCH(EOS_SESSIONS_REGISTERPLAYERS_API_LATEST, 3);
 
-					Options.PlayersToRegisterCount = 1;
+	Options.PlayersToRegisterCount = 1;
 
-					TArray<EOS_ProductUserId> PlayersToRegister = { GetProductUserIdChecked(Params.LocalAccountId) };
-					Options.PlayersToRegister = PlayersToRegister.GetData();
+	TArray<EOS_ProductUserId> PlayersToRegister = { GetProductUserIdChecked(Params.LocalAccountId) };
+	Options.PlayersToRegister = PlayersToRegister.GetData();
 
-					const FTCHARToUTF8 SessionNameUtf8(*Params.SessionName.ToString());
-					Options.SessionName = SessionNameUtf8.Get();
+	const FTCHARToUTF8 SessionNameUtf8(*Params.SessionName.ToString());
+	Options.SessionName = SessionNameUtf8.Get();
 
-					EOS_Async(EOS_Sessions_RegisterPlayers, SessionsHandle, Options,
-						[this, Promise = MoveTemp(Promise), Params](const EOS_Sessions_RegisterPlayersCallbackInfo* Result) mutable
-						{
-							if (Result->ResultCode != EOS_EResult::EOS_Success)
-							{
-								UE_LOG(LogOnlineServices, Warning, TEXT("EOS_Sessions_RegisterPlayers failed with result [%s]"), *LexToString(Result->ResultCode));
-								Promise.EmplaceValue(Errors::FromEOSResult(Result->ResultCode));
-								return;
-							}
-
-							Promise.EmplaceValue(AddSessionMemberInternal(Params));
-						});
-				}
-				else
-				{
-					Promise.EmplaceValue(Errors::InvalidUser());
-				}
-			});
+	EOS_Async(EOS_Sessions_RegisterPlayers, SessionsHandle, Options,
+		[this, Promise = MoveTemp(Promise), Params](const EOS_Sessions_RegisterPlayersCallbackInfo* Result) mutable
+	{
+		if (Result->ResultCode != EOS_EResult::EOS_Success)
+		{
+			UE_LOG(LogOnlineServices, Warning, TEXT("EOS_Sessions_RegisterPlayers failed with result [%s]"), *LexToString(Result->ResultCode));
+			Promise.EmplaceValue(Errors::FromEOSResult(Result->ResultCode));
+			return;
+		}
+
+		Promise.EmplaceValue(AddSessionMemberInternal(Params));
+	});
 
 	return Future;
 }
@@ -1563,41 +1552,31 @@ TFuture<TOnlineResult<FRemoveSessionMember>> FSessionsEOSGS::RemoveSessionMember
 	TPromise<TOnlineResult<FRemoveSessionMember>> Promise;
 	TFuture<TOnlineResult<FRemoveSessionMember>> Future = Promise.GetFuture();
 
-	Services.Get<IEpicProductUserIdResolver>()->ResolveAccountIds(Params.LocalAccountId, { GetProductUserIdChecked(Params.LocalAccountId) })
-		.Next([this, Promise = MoveTemp(Promise), Params](TArray<FAccountId> ResolvedAccountIds) mutable
-			{
-				if (!ResolvedAccountIds.IsEmpty())
-				{
-					EOS_Sessions_UnregisterPlayersOptions Options = { };
-					Options.ApiVersion = 2;
-					UE_EOS_CHECK_API_MISMATCH(EOS_SESSIONS_UNREGISTERPLAYERS_API_LATEST, 2);
+	
+	EOS_Sessions_UnregisterPlayersOptions Options = { };
+	Options.ApiVersion = 2;
+	UE_EOS_CHECK_API_MISMATCH(EOS_SESSIONS_UNREGISTERPLAYERS_API_LATEST, 2);
 
-					Options.PlayersToUnregisterCount = 1;
+	Options.PlayersToUnregisterCount = 1;
 
-					TArray<EOS_ProductUserId> PlayersToUnregister{ GetProductUserIdChecked(Params.LocalAccountId) };
-					Options.PlayersToUnregister = PlayersToUnregister.GetData();
+	TArray<EOS_ProductUserId> PlayersToUnregister{ GetProductUserIdChecked(Params.LocalAccountId) };
+	Options.PlayersToUnregister = PlayersToUnregister.GetData();
 
-					const FTCHARToUTF8 SessionNameUtf8(*Params.SessionName.ToString());
-					Options.SessionName = SessionNameUtf8.Get();
+	const FTCHARToUTF8 SessionNameUtf8(*Params.SessionName.ToString());
+	Options.SessionName = SessionNameUtf8.Get();
 
-					EOS_Async(EOS_Sessions_UnregisterPlayers, SessionsHandle, Options,
-						[this, Promise = MoveTemp(Promise), Params](const EOS_Sessions_UnregisterPlayersCallbackInfo* Result) mutable
-						{
-							if (Result->ResultCode != EOS_EResult::EOS_Success)
-							{
-								UE_LOG(LogOnlineServices, Warning, TEXT("EOS_Sessions_UnregisterPlayers failed with result [%s]"), *LexToString(Result->ResultCode));
-								Promise.EmplaceValue(Errors::FromEOSResult(Result->ResultCode));
-								return;
-							}
-
-							Promise.EmplaceValue(RemoveSessionMemberInternal(Params));
-						});
-				}
-				else
-				{
-					Promise.EmplaceValue(Errors::InvalidUser());
-				}
-			});
+	EOS_Async(EOS_Sessions_UnregisterPlayers, SessionsHandle, Options,
+		[this, Promise = MoveTemp(Promise), Params](const EOS_Sessions_UnregisterPlayersCallbackInfo* Result) mutable
+	{
+		if (Result->ResultCode != EOS_EResult::EOS_Success)
+		{
+			UE_LOG(LogOnlineServices, Warning, TEXT("EOS_Sessions_UnregisterPlayers failed with result [%s]"), *LexToString(Result->ResultCode));
+			Promise.EmplaceValue(Errors::FromEOSResult(Result->ResultCode));
+			return;
+		}
+
+		Promise.EmplaceValue(RemoveSessionMemberInternal(Params));
+	});
 
 	return Future;
 }
