Clear session cache when EOS_Sessions_DestroySession returns EOS_NotFound in FSessionsEOSGS::LeaveSessionImpl

#Synced-CL 49927350

diff --git a/Engine/Plugins/Online/OnlineServicesEOSGS/Source/Private/Online/SessionsEOSGS.cpp b/Engine/Plugins/Online/OnlineServicesEOSGS/Source/Private/Online/SessionsEOSGS.cpp
--- a/Engine/Plugins/Online/OnlineServicesEOSGS/Source/Private/Online/SessionsEOSGS.cpp
+++ b/Engine/Plugins/Online/OnlineServicesEOSGS/Source/Private/Online/SessionsEOSGS.cpp
@@ -936,29 +936,32 @@ TFuture<TOnlineResult<FLeaveSession>> FSessionsEOSGS::LeaveSessionImpl(const FLe
 	EOS_Async(EOS_Sessions_DestroySession, SessionsHandle, DestroySessionOptions,
 		[this, Promise = MoveTemp(Promise), Params](const EOS_Sessions_DestroySessionCallbackInfo* Result) mutable
 		{
-			if (Result->ResultCode != EOS_EResult::EOS_Success)
-			{
-				UE_LOG(LogOnlineServices, Warning, TEXT("EOS_Sessions_DestroySession failed with result [%s]"), *LexToString(Result->ResultCode));
-				Promise.EmplaceValue(Errors::FromEOSResult(Result->ResultCode));
-				return;
-			}
-
-			TOnlineResult<FGetSessionByName> GetSessionByNameResult = GetSessionByName({ Params.SessionName });
-			if (GetSessionByNameResult.IsOk())
+			if (Result->ResultCode == EOS_EResult::EOS_Success || Result->ResultCode == EOS_EResult::EOS_NotFound)
 			{
-				TSharedRef<const ISession> FoundSession = GetSessionByNameResult.GetOkValue().Session;
+				TOnlineResult<FGetSessionByName> GetSessionByNameResult = GetSessionByName({ Params.SessionName });
+				if (GetSessionByNameResult.IsOk())
+				{
+					TSharedRef<const ISession> FoundSession = GetSessionByNameResult.GetOkValue().Session;
 
-				ClearSessionReferences(FoundSession->GetSessionId(), Params.SessionName, Params.LocalAccountId);
+					ClearSessionReferences(FoundSession->GetSessionId(), Params.SessionName, Params.LocalAccountId);
+				}
+				else if (GetSessionByNameResult.GetErrorValue() != Errors::NotFound())
+				{
+					Promise.EmplaceValue(GetSessionByNameResult.GetErrorValue());
+					return; 
+				}
 
 				Promise.EmplaceValue(FLeaveSession::Result{ });
 
 				FSessionLeft SessionLeftEvent;
 				SessionLeftEvent.LocalAccountId = Params.LocalAccountId;
-				SessionEvents.OnSessionLeft.Broadcast(SessionLeftEvent);
+				SessionEvents.OnSessionLeft.Broadcast(SessionLeftEvent);	
 			}
 			else
 			{
-				Promise.EmplaceValue(GetSessionByNameResult.GetErrorValue());
+				UE_LOG(LogOnlineServices, Warning, TEXT("EOS_Sessions_DestroySession failed with result [%s]"), *LexToString(Result->ResultCode));
+				Promise.EmplaceValue(Errors::FromEOSResult(Result->ResultCode));
+				return;
 			}
 		});
 
