Add brackets to EOS P2P connection string in OnlineServicesEOSGS

#Synced-CL 51165957

diff --git a/Engine/Plugins/Online/OnlineServicesEOSGS/Source/Private/Online/OnlineServicesEOSGS.cpp b/Engine/Plugins/Online/OnlineServicesEOSGS/Source/Private/Online/OnlineServicesEOSGS.cpp
--- a/Engine/Plugins/Online/OnlineServicesEOSGS/Source/Private/Online/OnlineServicesEOSGS.cpp
+++ b/Engine/Plugins/Online/OnlineServicesEOSGS/Source/Private/Online/OnlineServicesEOSGS.cpp
@@ -110,9 +110,9 @@ TOnlineResult<FGetResolvedConnectString> FOnlineServicesEOSGS::GetResolvedConnec
 				if (Lobby->LobbyId == Params.LobbyId)
 				{
 #if WITH_ENGINE
-					//It should look like this: "EOS:0002aeeb5b2d4388a3752dd6d31222ec"
+					//It should look like this: "[EOS:0002aeeb5b2d4388a3752dd6d31222ec]"
 					FInternetAddrEOS TempAddr(GetProductUserIdChecked(Lobby->OwnerAccountId));
-					return TOnlineResult<FGetResolvedConnectString>({ TempAddr.ToString(true) });
+					return TOnlineResult<FGetResolvedConnectString>({ FString::Printf(TEXT("[%s]"), *TempAddr.ToString(true)) });
 #else
 					return TOnlineResult<FGetResolvedConnectString>(Errors::NotImplemented());
 #endif
@@ -138,9 +138,9 @@ TOnlineResult<FGetResolvedConnectString> FOnlineServicesEOSGS::GetResolvedConnec
 #if WITH_ENGINE
 			TSharedRef<const ISession> Session = Result.GetOkValue().Session;
 
-			//It should look like this: "EOS:0002aeeb5b2d4388a3752dd6d31222ec"
+			//It should look like this: "[EOS:0002aeeb5b2d4388a3752dd6d31222ec]"
 			FInternetAddrEOS TempAddr(GetProductUserIdChecked(Session->GetOwnerAccountId()));
-			return TOnlineResult<FGetResolvedConnectString>({ TempAddr.ToString(true) });
+			return TOnlineResult<FGetResolvedConnectString>({ FString::Printf(TEXT("[%s]"), *TempAddr.ToString(true)) });
 #else
 			return TOnlineResult<FGetResolvedConnectString>(Errors::NotImplemented());
 #endif
